Copying WEX-Staging @ (WEX/Main @ 3740665) to //UE4/Main

diff --git a/Engine/Plugins/Editor/AssetManagerEditor/Source/AssetManagerEditor/Private/PrimaryAssetTypeCustomization.cpp b/Engine/Plugins/Editor/AssetManagerEditor/Source/AssetManagerEditor/Private/PrimaryAssetTypeCustomization.cpp
--- a/Engine/Plugins/Editor/AssetManagerEditor/Source/AssetManagerEditor/Private/PrimaryAssetTypeCustomization.cpp
+++ b/Engine/Plugins/Editor/AssetManagerEditor/Source/AssetManagerEditor/Private/PrimaryAssetTypeCustomization.cpp
@@ -43,7 +43,7 @@ void FPrimaryAssetTypeCustomization::CustomizeHeader(TSharedRef<class IPropertyH
 	.MinDesiredWidth(250.0f)
 	.MaxDesiredWidth(0.0f)
 	[
-		PropertyCustomizationHelpers::MakePropertyComboBox(InStructPropertyHandle, FOnGetPropertyComboBoxStrings::CreateStatic(&IAssetManagerEditorModule::GeneratePrimaryAssetTypeComboBoxStrings, bAllowClear))
+		PropertyCustomizationHelpers::MakePropertyComboBox(InStructPropertyHandle, FOnGetPropertyComboBoxStrings::CreateStatic(&IAssetManagerEditorModule::GeneratePrimaryAssetTypeComboBoxStrings, bAllowClear, false))
 	];
 }
 
diff --git a/Engine/Plugins/Editor/AssetManagerEditor/Source/AssetManagerEditor/Private/ReferenceViewer/EdGraphNode_Reference.cpp b/Engine/Plugins/Editor/AssetManagerEditor/Source/AssetManagerEditor/Private/ReferenceViewer/EdGraphNode_Reference.cpp
new file mode 100644
--- /dev/null
+++ b/Engine/Plugins/Editor/AssetManagerEditor/Source/AssetManagerEditor/Private/ReferenceViewer/EdGraphNode_Reference.cpp
@@ -0,0 +1,255 @@
+// Copyright 1998-2017 Epic Games, Inc. All Rights Reserved.
+
+#include "EdGraphNode_Reference.h"
+#include "GenericPlatform/GenericPlatformFile.h"
+#include "EdGraph/EdGraphPin.h"
+#include "HAL/PlatformFilemanager.h"
+
+#define LOCTEXT_NAMESPACE "ReferenceViewer"
+
+//////////////////////////////////////////////////////////////////////////
+// UEdGraphNode_Reference
+
+UEdGraphNode_Reference::UEdGraphNode_Reference(const FObjectInitializer& ObjectInitializer)
+	: Super(ObjectInitializer)
+{
+	DependencyPin = NULL;
+	ReferencerPin = NULL;
+	bIsCollapsed = false;
+	bIsPackage = false;
+	bIsPrimaryAsset = false;
+	bUsesThumbnail = false;
+}
+
+void UEdGraphNode_Reference::SetupReferenceNode(const FIntPoint& NodeLoc, const TArray<FAssetIdentifier>& NewIdentifiers, const FAssetData& InAssetData)
+{
+	check(NewIdentifiers.Num() > 0);
+
+	NodePosX = NodeLoc.X;
+	NodePosY = NodeLoc.Y;
+
+	Identifiers = NewIdentifiers;
+	const FAssetIdentifier& First = NewIdentifiers[0];
+	FString ShortPackageName = FPackageName::GetLongPackageAssetName(First.PackageName.ToString());
+
+	bIsCollapsed = false;
+	bIsPackage = true;
+	
+	FPrimaryAssetId PrimaryAssetID = NewIdentifiers[0].GetPrimaryAssetId();
+	if (PrimaryAssetID.IsValid())
+	{
+		ShortPackageName = PrimaryAssetID.ToString();
+		bIsPackage = false;
+		bIsPrimaryAsset = true;
+	}
+	else if (NewIdentifiers[0].IsValue())
+	{
+		ShortPackageName = FString::Printf(TEXT("%s::%s"), *First.ObjectName.ToString(), *First.ValueName.ToString());
+		bIsPackage = false;
+	}
+
+	if (NewIdentifiers.Num() == 1 )
+	{
+		if (bIsPackage)
+		{
+			NodeComment = First.PackageName.ToString();
+		}
+		NodeTitle = FText::FromString(ShortPackageName);
+	}
+	else
+	{
+		NodeComment = FText::Format(LOCTEXT("ReferenceNodeMultiplePackagesTitle", "{0} nodes"), FText::AsNumber(NewIdentifiers.Num())).ToString();
+		NodeTitle = FText::Format(LOCTEXT("ReferenceNodeMultiplePackagesComment", "{0} and {1} others"), FText::FromString(ShortPackageName), FText::AsNumber(NewIdentifiers.Num() - 1));
+	}
+	
+	CacheAssetData(InAssetData);
+	AllocateDefaultPins();
+}
+
+void UEdGraphNode_Reference::SetReferenceNodeCollapsed(const FIntPoint& NodeLoc, int32 InNumReferencesExceedingMax)
+{
+	NodePosX = NodeLoc.X;
+	NodePosY = NodeLoc.Y;
+
+	Identifiers.Empty();
+	bIsCollapsed = true;
+	bUsesThumbnail = false;
+	NodeComment = FText::Format(LOCTEXT("ReferenceNodeCollapsedMessage", "{0} other nodes"), FText::AsNumber(InNumReferencesExceedingMax)).ToString();
+
+	NodeTitle = LOCTEXT("ReferenceNodeCollapsedTitle", "Collapsed nodes");
+	CacheAssetData(FAssetData());
+	AllocateDefaultPins();
+}
+
+void UEdGraphNode_Reference::AddReferencer(UEdGraphNode_Reference* ReferencerNode)
+{
+	UEdGraphPin* ReferencerDependencyPin = ReferencerNode->GetDependencyPin();
+
+	if ( ensure(ReferencerDependencyPin) )
+	{
+		ReferencerDependencyPin->bHidden = false;
+		ReferencerPin->bHidden = false;
+		ReferencerPin->MakeLinkTo(ReferencerDependencyPin);
+	}
+}
+
+FAssetIdentifier UEdGraphNode_Reference::GetIdentifier() const
+{
+	if (Identifiers.Num() > 0)
+	{
+		return Identifiers[0];
+	}
+
+	return FAssetIdentifier();
+}
+
+void UEdGraphNode_Reference::GetAllIdentifiers(TArray<FAssetIdentifier>& OutIdentifiers) const
+{
+	OutIdentifiers.Append(Identifiers);
+}
+
+void UEdGraphNode_Reference::GetAllPackageNames(TArray<FName>& OutPackageNames) const
+{
+	for (const FAssetIdentifier& AssetId : Identifiers)
+	{
+		if (AssetId.IsPackage())
+		{
+			OutPackageNames.AddUnique(AssetId.PackageName);
+		}
+	}
+}
+
+UEdGraph_ReferenceViewer* UEdGraphNode_Reference::GetReferenceViewerGraph() const
+{
+	return Cast<UEdGraph_ReferenceViewer>( GetGraph() );
+}
+
+FText UEdGraphNode_Reference::GetNodeTitle(ENodeTitleType::Type TitleType) const
+{
+	return NodeTitle;
+}
+
+FLinearColor UEdGraphNode_Reference::GetNodeTitleColor() const
+{
+	if (bIsPrimaryAsset)
+	{
+		return FLinearColor(0.2f, 0.8f, 0.2f);
+	}
+	else if (bIsPackage)
+	{
+		return FLinearColor(0.4f, 0.62f, 1.0f);
+	}
+	else if (bIsCollapsed)
+	{
+		return FLinearColor(0.55f, 0.55f, 0.55f);
+	}
+	else 
+	{
+		return FLinearColor(0.0f, 0.55f, 0.62f);
+	}
+}
+
+FText UEdGraphNode_Reference::GetTooltipText() const
+{
+	FString TooltipString;
+	for (const FAssetIdentifier& AssetId : Identifiers)
+	{
+		if (!TooltipString.IsEmpty())
+		{
+			TooltipString.Append(TEXT("\n"));
+		}
+		TooltipString.Append(AssetId.ToString());
+	}
+	return FText::FromString(TooltipString);
+}
+
+void UEdGraphNode_Reference::AllocateDefaultPins()
+{
+	ReferencerPin = CreatePin( EEdGraphPinDirection::EGPD_Input, NAME_None, NAME_None);
+	DependencyPin = CreatePin( EEdGraphPinDirection::EGPD_Output, NAME_None, NAME_None);
+
+	ReferencerPin->bHidden = true;
+	DependencyPin->bHidden = true;
+}
+
+UObject* UEdGraphNode_Reference::GetJumpTargetForDoubleClick() const
+{
+	if (Identifiers.Num() > 0 )
+	{
+		GetReferenceGraph()->SetGraphRoot(Identifiers, FIntPoint(NodePosX, NodePosY));
+		GetReferenceGraph()->RebuildGraph();
+	}
+	return NULL;
+}
+
+UEdGraphPin* UEdGraphNode_Reference::GetDependencyPin()
+{
+	return DependencyPin;
+}
+
+UEdGraphPin* UEdGraphNode_Reference::GetReferencerPin()
+{
+	return ReferencerPin;
+}
+
+void UEdGraphNode_Reference::CacheAssetData(const FAssetData& AssetData)
+{
+	if ( AssetData.IsValid() && IsPackage() )
+	{
+		bUsesThumbnail = true;
+		CachedAssetData = AssetData;
+	}
+	else
+	{
+		CachedAssetData = FAssetData();
+		bUsesThumbnail = false;
+
+		if (Identifiers.Num() == 1 )
+		{
+			const FString PackageNameStr = Identifiers[0].PackageName.ToString();
+			if ( FPackageName::IsValidLongPackageName(PackageNameStr, true) )
+			{
+				if ( PackageNameStr.StartsWith(TEXT("/Script")) )
+				{
+					CachedAssetData.AssetClass = FName(TEXT("Code"));
+				}
+				else
+				{
+					const FString PotentiallyMapFilename = FPackageName::LongPackageNameToFilename(PackageNameStr, FPackageName::GetMapPackageExtension());
+					const bool bIsMapPackage = FPlatformFileManager::Get().GetPlatformFile().FileExists(*PotentiallyMapFilename);
+					if ( bIsMapPackage )
+					{
+						CachedAssetData.AssetClass = FName(TEXT("World"));
+					}
+				}
+			}
+		}
+		else
+		{
+			CachedAssetData.AssetClass = FName(TEXT("Multiple Nodes"));
+		}
+	}
+
+}
+
+FAssetData UEdGraphNode_Reference::GetAssetData() const
+{
+	return CachedAssetData;
+}
+
+bool UEdGraphNode_Reference::UsesThumbnail() const
+{
+	return bUsesThumbnail;
+}
+
+bool UEdGraphNode_Reference::IsPackage() const
+{
+	return bIsPackage;
+}
+
+bool UEdGraphNode_Reference::IsCollapsed() const
+{
+	return bIsCollapsed;
+}
+
+#undef LOCTEXT_NAMESPACE
diff --git a/Engine/Plugins/Editor/AssetManagerEditor/Source/AssetManagerEditor/Private/ReferenceViewer/EdGraphNode_Reference.h b/Engine/Plugins/Editor/AssetManagerEditor/Source/AssetManagerEditor/Private/ReferenceViewer/EdGraphNode_Reference.h
new file mode 100644
--- /dev/null
+++ b/Engine/Plugins/Editor/AssetManagerEditor/Source/AssetManagerEditor/Private/ReferenceViewer/EdGraphNode_Reference.h
@@ -0,0 +1,66 @@
+// Copyright 1998-2017 Epic Games, Inc. All Rights Reserved.
+
+#pragma once
+
+#include "CoreMinimal.h"
+#include "UObject/ObjectMacros.h"
+#include "EdGraph/EdGraphNode.h"
+#include "AssetData.h"
+#include "EdGraph_ReferenceViewer.h"
+#include "EdGraphNode_Reference.generated.h"
+
+class UEdGraphPin;
+
+UCLASS()
+class UEdGraphNode_Reference : public UEdGraphNode
+{
+	GENERATED_UCLASS_BODY()
+
+	virtual void SetupReferenceNode(const FIntPoint& NodeLoc, const TArray<FAssetIdentifier>& NewIdentifiers, const FAssetData& InAssetData);
+	virtual void SetReferenceNodeCollapsed(const FIntPoint& NodeLoc, int32 InNumReferencesExceedingMax);
+	virtual void AddReferencer(class UEdGraphNode_Reference* ReferencerNode);
+	// Returns first asset identifier
+	virtual FAssetIdentifier GetIdentifier() const;
+	virtual void GetAllIdentifiers(TArray<FAssetIdentifier>& OutIdentifiers) const;
+	// Returns only the packages in this node, skips searchable names
+	virtual void GetAllPackageNames(TArray<FName>& OutPackageNames) const;
+	class UEdGraph_ReferenceViewer* GetReferenceViewerGraph() const;
+
+	// UEdGraphNode implementation
+	virtual FText GetNodeTitle(ENodeTitleType::Type TitleType) const override;
+	FLinearColor GetNodeTitleColor() const override;
+	virtual FText GetTooltipText() const override;
+	virtual void AllocateDefaultPins() override;
+	virtual UObject* GetJumpTargetForDoubleClick() const override;
+	// End UEdGraphNode implementation
+
+	void CacheAssetData(const FAssetData& AssetData);
+
+	bool UsesThumbnail() const;
+	bool IsPackage() const;
+	bool IsCollapsed() const;
+	FAssetData GetAssetData() const;
+
+	virtual UEdGraphPin* GetDependencyPin();
+	virtual UEdGraphPin* GetReferencerPin();
+protected:
+	class UEdGraph_ReferenceViewer* GetReferenceGraph() const
+	{
+		return CastChecked<UEdGraph_ReferenceViewer>(GetOuter());
+	}
+
+private:
+	TArray<FAssetIdentifier> Identifiers;
+	FText NodeTitle;
+
+	bool bUsesThumbnail;
+	bool bIsPackage;
+	bool bIsPrimaryAsset;
+	bool bIsCollapsed;
+	FAssetData CachedAssetData;
+
+	UEdGraphPin* DependencyPin;
+	UEdGraphPin* ReferencerPin;
+};
+
+
diff --git a/Engine/Plugins/Editor/AssetManagerEditor/Source/AssetManagerEditor/Private/ReferenceViewer/EdGraph_ReferenceViewer.cpp b/Engine/Plugins/Editor/AssetManagerEditor/Source/AssetManagerEditor/Private/ReferenceViewer/EdGraph_ReferenceViewer.cpp
new file mode 100644
--- /dev/null
+++ b/Engine/Plugins/Editor/AssetManagerEditor/Source/AssetManagerEditor/Private/ReferenceViewer/EdGraph_ReferenceViewer.cpp
@@ -0,0 +1,610 @@
+// Copyright 1998-2017 Epic Games, Inc. All Rights Reserved.
+
+#include "EdGraph_ReferenceViewer.h"
+#include "EdGraph/EdGraphPin.h"
+#include "EdGraphNode_Reference.h"
+#include "ARFilter.h"
+#include "AssetRegistryModule.h"
+#include "AssetThumbnail.h"
+#include "SReferenceViewer.h"
+#include "SReferenceNode.h"
+#include "GraphEditor.h"
+#include "ICollectionManager.h"
+#include "CollectionManagerModule.h"
+#include "AssetManagerEditorModule.h"
+#include "Engine/AssetManager.h"
+
+UEdGraph_ReferenceViewer::UEdGraph_ReferenceViewer(const FObjectInitializer& ObjectInitializer)
+	: Super(ObjectInitializer)
+{
+	AssetThumbnailPool = MakeShareable( new FAssetThumbnailPool(1024) );
+
+	MaxSearchDepth = 1;
+	MaxSearchBreadth = 15;
+
+	bLimitSearchDepth = true;
+	bLimitSearchBreadth = true;
+	bIsShowSoftReferences = true;
+	bIsShowHardReferences = true;
+	bIsShowManagementReferences = false;
+	bIsShowSearchableNames = false;
+	bIsShowNativePackages = false;
+}
+
+void UEdGraph_ReferenceViewer::BeginDestroy()
+{
+	if ( AssetThumbnailPool.IsValid() )
+	{
+		AssetThumbnailPool->ReleaseResources();
+		AssetThumbnailPool.Reset();
+	}
+
+	Super::BeginDestroy();
+}
+
+void UEdGraph_ReferenceViewer::SetGraphRoot(const TArray<FAssetIdentifier>& GraphRootIdentifiers, const FIntPoint& GraphRootOrigin)
+{
+	CurrentGraphRootIdentifiers = GraphRootIdentifiers;
+	CurrentGraphRootOrigin = GraphRootOrigin;
+
+	// If we're focused on a searchable name, enable that flag
+	for (const FAssetIdentifier& AssetId : GraphRootIdentifiers)
+	{
+		if (AssetId.IsValue())
+		{
+			bIsShowSearchableNames = true;
+		}
+		else if (AssetId.GetPrimaryAssetId().IsValid())
+		{
+			if (UAssetManager::IsValid())
+			{
+				UAssetManager::Get().UpdateManagementDatabase();
+			}
+			
+			bIsShowManagementReferences = true;
+		}
+	}
+}
+
+const TArray<FAssetIdentifier>& UEdGraph_ReferenceViewer::GetCurrentGraphRootIdentifiers() const
+{
+	return CurrentGraphRootIdentifiers;
+}
+
+void UEdGraph_ReferenceViewer::SetReferenceViewer(TSharedPtr<SReferenceViewer> InViewer)
+{
+	ReferenceViewer = InViewer;
+}
+
+bool UEdGraph_ReferenceViewer::GetSelectedAssetsForMenuExtender(const class UEdGraphNode* Node, TArray<FAssetIdentifier>& SelectedAssets) const
+{
+	if (!ReferenceViewer.IsValid())
+	{
+		return false;
+	}
+	TSharedPtr<SGraphEditor> GraphEditor = ReferenceViewer.Pin()->GetGraphEditor();
+
+	if (!GraphEditor.IsValid())
+	{
+		return false;
+	}
+
+	TSet<UObject*> SelectedNodes = GraphEditor->GetSelectedNodes();
+	for (FGraphPanelSelectionSet::TConstIterator It(SelectedNodes); It; ++It)
+	{
+		if (UEdGraphNode_Reference* ReferenceNode = Cast<UEdGraphNode_Reference>(*It))
+		{
+			if (!ReferenceNode->IsCollapsed())
+			{
+				SelectedAssets.Add(ReferenceNode->GetIdentifier());
+			}
+		}
+	}
+	return true;
+}
+
+UEdGraphNode_Reference* UEdGraph_ReferenceViewer::RebuildGraph()
+{
+	RemoveAllNodes();
+	UEdGraphNode_Reference* NewRootNode = ConstructNodes(CurrentGraphRootIdentifiers, CurrentGraphRootOrigin);
+	NotifyGraphChanged();
+
+	return NewRootNode;
+}
+
+bool UEdGraph_ReferenceViewer::IsSearchDepthLimited() const
+{
+	return bLimitSearchDepth;
+}
+
+bool UEdGraph_ReferenceViewer::IsSearchBreadthLimited() const
+{
+	return bLimitSearchBreadth;
+}
+
+bool UEdGraph_ReferenceViewer::IsShowSoftReferences() const
+{
+	return bIsShowSoftReferences;
+}
+
+bool UEdGraph_ReferenceViewer::IsShowHardReferences() const
+{
+	return bIsShowHardReferences;
+}
+
+bool UEdGraph_ReferenceViewer::IsShowManagementReferences() const
+{
+	return bIsShowManagementReferences;
+}
+
+bool UEdGraph_ReferenceViewer::IsShowSearchableNames() const
+{
+	return bIsShowSearchableNames;
+}
+
+bool UEdGraph_ReferenceViewer::IsShowNativePackages() const
+{
+	return bIsShowNativePackages;
+}
+
+void UEdGraph_ReferenceViewer::SetSearchDepthLimitEnabled(bool newEnabled)
+{
+	bLimitSearchDepth = newEnabled;
+}
+
+void UEdGraph_ReferenceViewer::SetSearchBreadthLimitEnabled(bool newEnabled)
+{
+	bLimitSearchBreadth = newEnabled;
+}
+
+void UEdGraph_ReferenceViewer::SetShowSoftReferencesEnabled(bool newEnabled)
+{
+	bIsShowSoftReferences = newEnabled;
+}
+
+void UEdGraph_ReferenceViewer::SetShowHardReferencesEnabled(bool newEnabled)
+{
+	bIsShowHardReferences = newEnabled;
+}
+
+void UEdGraph_ReferenceViewer::SetShowManagementReferencesEnabled(bool newEnabled)
+{
+	bIsShowManagementReferences = newEnabled;
+}
+
+void UEdGraph_ReferenceViewer::SetShowSearchableNames(bool newEnabled)
+{
+	bIsShowSearchableNames = newEnabled;
+}
+
+void UEdGraph_ReferenceViewer::SetShowNativePackages(bool newEnabled)
+{
+	bIsShowNativePackages = newEnabled;
+}
+
+int32 UEdGraph_ReferenceViewer::GetSearchDepthLimit() const
+{
+	return MaxSearchDepth;
+}
+
+int32 UEdGraph_ReferenceViewer::GetSearchBreadthLimit() const
+{
+	return MaxSearchBreadth;
+}
+
+void UEdGraph_ReferenceViewer::SetSearchDepthLimit(int32 NewDepthLimit)
+{
+	MaxSearchDepth = NewDepthLimit;
+}
+
+void UEdGraph_ReferenceViewer::SetSearchBreadthLimit(int32 NewBreadthLimit)
+{
+	MaxSearchBreadth = NewBreadthLimit;
+}
+
+FName UEdGraph_ReferenceViewer::GetCurrentCollectionFilter() const
+{
+	return CurrentCollectionFilter;
+}
+
+void UEdGraph_ReferenceViewer::SetCurrentCollectionFilter(FName NewFilter)
+{
+	CurrentCollectionFilter = NewFilter;
+}
+
+bool UEdGraph_ReferenceViewer::GetEnableCollectionFilter() const
+{
+	return bEnableCollectionFilter;
+}
+
+void UEdGraph_ReferenceViewer::SetEnableCollectionFilter(bool bEnabled)
+{
+	bEnableCollectionFilter = bEnabled;
+}
+
+EAssetRegistryDependencyType::Type UEdGraph_ReferenceViewer::GetReferenceSearchFlags(bool bHardOnly) const
+{
+	int32 ReferenceFlags = 0;
+
+	if (bIsShowSoftReferences && !bHardOnly)
+	{
+		ReferenceFlags |= EAssetRegistryDependencyType::Soft;
+	}
+	if (bIsShowHardReferences)
+	{
+		ReferenceFlags |= EAssetRegistryDependencyType::Hard;
+	}
+	if (bIsShowSearchableNames && !bHardOnly)
+	{
+		ReferenceFlags |= EAssetRegistryDependencyType::SearchableName;
+	}
+	if (bIsShowManagementReferences)
+	{
+		ReferenceFlags |= EAssetRegistryDependencyType::HardManage;
+		if (!bHardOnly)
+		{
+			ReferenceFlags |= EAssetRegistryDependencyType::SoftManage;
+		}
+	}
+	
+	return (EAssetRegistryDependencyType::Type)ReferenceFlags;
+}
+
+UEdGraphNode_Reference* UEdGraph_ReferenceViewer::ConstructNodes(const TArray<FAssetIdentifier>& GraphRootIdentifiers, const FIntPoint& GraphRootOrigin )
+{
+	UEdGraphNode_Reference* RootNode = NULL;
+
+	if (GraphRootIdentifiers.Num() > 0 )
+	{
+		TSet<FName> AllowedPackageNames;
+		if (ShouldFilterByCollection())
+		{
+			FCollectionManagerModule& CollectionManagerModule = FCollectionManagerModule::GetModule();
+			TArray<FName> AssetPaths;
+			CollectionManagerModule.Get().GetAssetsInCollection(CurrentCollectionFilter, ECollectionShareType::CST_All, AssetPaths);
+			AllowedPackageNames.Reserve(AssetPaths.Num());
+			for (FName AssetPath : AssetPaths)
+			{
+				AllowedPackageNames.Add(FName(*FPackageName::ObjectPathToPackageName(AssetPath.ToString())));
+			}
+		}
+
+		TMap<FAssetIdentifier, int32> ReferencerNodeSizes;
+		TSet<FAssetIdentifier> VisitedReferencerSizeNames;
+		int32 ReferencerDepth = 1;
+		RecursivelyGatherSizes(/*bReferencers=*/true, GraphRootIdentifiers, AllowedPackageNames, ReferencerDepth, VisitedReferencerSizeNames, ReferencerNodeSizes);
+
+		TMap<FAssetIdentifier, int32> DependencyNodeSizes;
+		TSet<FAssetIdentifier> VisitedDependencySizeNames;
+		int32 DependencyDepth = 1;
+		RecursivelyGatherSizes(/*bReferencers=*/false, GraphRootIdentifiers, AllowedPackageNames, DependencyDepth, VisitedDependencySizeNames, DependencyNodeSizes);
+
+		TSet<FName> AllPackageNames;
+
+		auto AddPackage = [](const FAssetIdentifier& AssetId, TSet<FName>& PackageNames)
+		{ 
+			// Only look for asset data if this is a package
+			if (!AssetId.IsValue())
+			{
+				PackageNames.Add(AssetId.PackageName);
+			}
+		};
+
+		for (const FAssetIdentifier& AssetId : VisitedReferencerSizeNames)
+		{
+			AddPackage(AssetId, AllPackageNames);
+		}
+
+		for (const FAssetIdentifier& AssetId : VisitedDependencySizeNames)
+		{
+			AddPackage(AssetId, AllPackageNames);
+		}
+
+		TMap<FName, FAssetData> PackagesToAssetDataMap;
+		GatherAssetData(AllPackageNames, PackagesToAssetDataMap);
+
+		// Create the root node
+		RootNode = CreateReferenceNode();
+		RootNode->SetupReferenceNode(GraphRootOrigin, GraphRootIdentifiers, PackagesToAssetDataMap.FindRef(GraphRootIdentifiers[0].PackageName));
+
+		TSet<FAssetIdentifier> VisitedReferencerNames;
+		int32 VisitedReferencerDepth = 1;
+		RecursivelyConstructNodes(/*bReferencers=*/true, RootNode, GraphRootIdentifiers, GraphRootOrigin, ReferencerNodeSizes, PackagesToAssetDataMap, AllowedPackageNames, VisitedReferencerDepth, VisitedReferencerNames);
+
+		TSet<FAssetIdentifier> VisitedDependencyNames;
+		int32 VisitedDependencyDepth = 1;
+		RecursivelyConstructNodes(/*bReferencers=*/false, RootNode, GraphRootIdentifiers, GraphRootOrigin, DependencyNodeSizes, PackagesToAssetDataMap, AllowedPackageNames, VisitedDependencyDepth, VisitedDependencyNames);
+	}
+
+	return RootNode;
+}
+
+int32 UEdGraph_ReferenceViewer::RecursivelyGatherSizes(bool bReferencers, const TArray<FAssetIdentifier>& Identifiers, const TSet<FName>& AllowedPackageNames, int32 CurrentDepth, TSet<FAssetIdentifier>& VisitedNames, TMap<FAssetIdentifier, int32>& OutNodeSizes) const
+{
+	check(Identifiers.Num() > 0);
+
+	VisitedNames.Append(Identifiers);
+
+	FAssetRegistryModule& AssetRegistryModule = FModuleManager::LoadModuleChecked<FAssetRegistryModule>("AssetRegistry");
+	TArray<FAssetIdentifier> ReferenceNames;
+
+	if ( bReferencers )
+	{
+		for (const FAssetIdentifier& AssetId : Identifiers)
+		{
+			AssetRegistryModule.Get().GetReferencers(AssetId, ReferenceNames, GetReferenceSearchFlags(false));
+		}
+	}
+	else
+	{
+		for (const FAssetIdentifier& AssetId : Identifiers)
+		{
+			AssetRegistryModule.Get().GetDependencies(AssetId, ReferenceNames, GetReferenceSearchFlags(false));
+		}
+	}
+
+	if (!bIsShowNativePackages)
+	{
+		auto RemoveNativePackage = [](const FAssetIdentifier& InAsset) { return InAsset.PackageName.ToString().StartsWith(TEXT("/Script")) && !InAsset.IsValue(); };
+
+		ReferenceNames.RemoveAll(RemoveNativePackage);
+	}
+
+	int32 NodeSize = 0;
+	if ( ReferenceNames.Num() > 0 && !ExceedsMaxSearchDepth(CurrentDepth) )
+	{
+		int32 NumReferencesMade = 0;
+		int32 NumReferencesExceedingMax = 0;
+
+		// Filter for our registry source
+		IAssetManagerEditorModule::Get().FilterAssetIdentifiersForCurrentRegistrySource(ReferenceNames, GetReferenceSearchFlags(false), !bReferencers);
+
+		// Since there are referencers, use the size of all your combined referencers.
+		// Do not count your own size since there could just be a horizontal line of nodes
+		for (FAssetIdentifier& AssetId : ReferenceNames)
+		{
+			if ( !VisitedNames.Contains(AssetId) && (!AssetId.IsPackage() || !ShouldFilterByCollection() || AllowedPackageNames.Contains(AssetId.PackageName)) )
+			{
+				if ( !ExceedsMaxSearchBreadth(NumReferencesMade) )
+				{
+					TArray<FAssetIdentifier> NewPackageNames;
+					NewPackageNames.Add(AssetId);
+					NodeSize += RecursivelyGatherSizes(bReferencers, NewPackageNames, AllowedPackageNames, CurrentDepth + 1, VisitedNames, OutNodeSizes);
+					NumReferencesMade++;
+				}
+				else
+				{
+					NumReferencesExceedingMax++;
+				}
+			}
+		}
+
+		if ( NumReferencesExceedingMax > 0 )
+		{
+			// Add one size for the collapsed node
+			NodeSize++;
+		}
+	}
+
+	if ( NodeSize == 0 )
+	{
+		// If you have no valid children, the node size is just 1 (counting only self to make a straight line)
+		NodeSize = 1;
+	}
+
+	OutNodeSizes.Add(Identifiers[0], NodeSize);
+	return NodeSize;
+}
+
+void UEdGraph_ReferenceViewer::GatherAssetData(const TSet<FName>& AllPackageNames, TMap<FName, FAssetData>& OutPackageToAssetDataMap) const
+{
+	// Take a guess to find the asset instead of searching for it. Most packages have a single asset in them with the same name as the package.
+	FAssetRegistryModule& AssetRegistryModule = FModuleManager::LoadModuleChecked<FAssetRegistryModule>("AssetRegistry");
+
+	FARFilter Filter;
+	for ( auto PackageIt = AllPackageNames.CreateConstIterator(); PackageIt; ++PackageIt )
+	{
+		const FString& PackageName = (*PackageIt).ToString();
+		const FString& PackagePath = PackageName + TEXT(".") + FPackageName::GetLongPackageAssetName(PackageName);
+		Filter.ObjectPaths.Add( FName(*PackagePath) );
+	}
+
+	TArray<FAssetData> AssetDataList;
+	AssetRegistryModule.Get().GetAssets(Filter, AssetDataList);
+	for ( auto AssetIt = AssetDataList.CreateConstIterator(); AssetIt; ++AssetIt )
+	{
+		OutPackageToAssetDataMap.Add((*AssetIt).PackageName, *AssetIt);
+	}
+}
+
+UEdGraphNode_Reference* UEdGraph_ReferenceViewer::RecursivelyConstructNodes(bool bReferencers, UEdGraphNode_Reference* RootNode, const TArray<FAssetIdentifier>& Identifiers, const FIntPoint& NodeLoc, const TMap<FAssetIdentifier, int32>& NodeSizes, const TMap<FName, FAssetData>& PackagesToAssetDataMap, const TSet<FName>& AllowedPackageNames, int32 CurrentDepth, TSet<FAssetIdentifier>& VisitedNames)
+{
+	check(Identifiers.Num() > 0);
+
+	VisitedNames.Append(Identifiers);
+
+	UEdGraphNode_Reference* NewNode = NULL;
+	if ( RootNode->GetIdentifier() == Identifiers[0] )
+	{
+		// Don't create the root node. It is already created!
+		NewNode = RootNode;
+	}
+	else
+	{
+		NewNode = CreateReferenceNode();
+		NewNode->SetupReferenceNode(NodeLoc, Identifiers, PackagesToAssetDataMap.FindRef(Identifiers[0].PackageName));
+	}
+
+	FAssetRegistryModule& AssetRegistryModule = FModuleManager::LoadModuleChecked<FAssetRegistryModule>("AssetRegistry");
+	TArray<FAssetIdentifier> ReferenceNames;
+	TArray<FAssetIdentifier> HardReferenceNames;
+	if ( bReferencers )
+	{
+		for (const FAssetIdentifier& AssetId : Identifiers)
+		{
+			AssetRegistryModule.Get().GetReferencers(AssetId, HardReferenceNames, GetReferenceSearchFlags(true));
+			AssetRegistryModule.Get().GetReferencers(AssetId, ReferenceNames, GetReferenceSearchFlags(false));
+		}
+	}
+	else
+	{
+		for (const FAssetIdentifier& AssetId : Identifiers)
+		{
+			AssetRegistryModule.Get().GetDependencies(AssetId, HardReferenceNames, GetReferenceSearchFlags(true));
+			AssetRegistryModule.Get().GetDependencies(AssetId, ReferenceNames, GetReferenceSearchFlags(false));
+		}
+	}
+
+	if (!bIsShowNativePackages)
+	{
+		auto RemoveNativePackage = [](const FAssetIdentifier& InAsset) { return InAsset.PackageName.ToString().StartsWith(TEXT("/Script")) && !InAsset.IsValue(); };
+
+		HardReferenceNames.RemoveAll(RemoveNativePackage);
+		ReferenceNames.RemoveAll(RemoveNativePackage);
+	}
+
+	if ( ReferenceNames.Num() > 0 && !ExceedsMaxSearchDepth(CurrentDepth) )
+	{
+		FIntPoint ReferenceNodeLoc = NodeLoc;
+
+		if ( bReferencers )
+		{
+			// Referencers go left
+			ReferenceNodeLoc.X -= 800;
+		}
+		else
+		{
+			// Dependencies go right
+			ReferenceNodeLoc.X += 800;
+		}
+
+		const int32 NodeSizeY = 200;
+		const int32 TotalReferenceSizeY = NodeSizes.FindChecked(Identifiers[0]) * NodeSizeY;
+
+		ReferenceNodeLoc.Y -= TotalReferenceSizeY * 0.5f;
+		ReferenceNodeLoc.Y += NodeSizeY * 0.5f;
+
+		int32 NumReferencesMade = 0;
+		int32 NumReferencesExceedingMax = 0;
+		
+		// Filter for our registry source
+		IAssetManagerEditorModule::Get().FilterAssetIdentifiersForCurrentRegistrySource(ReferenceNames, GetReferenceSearchFlags(false), !bReferencers);
+		IAssetManagerEditorModule::Get().FilterAssetIdentifiersForCurrentRegistrySource(HardReferenceNames, GetReferenceSearchFlags(false), !bReferencers);
+
+		for ( int32 RefIdx = 0; RefIdx < ReferenceNames.Num(); ++RefIdx )
+		{
+			FAssetIdentifier ReferenceName = ReferenceNames[RefIdx];
+
+			if ( !VisitedNames.Contains(ReferenceName) && (!ReferenceName.IsPackage() || !ShouldFilterByCollection() || AllowedPackageNames.Contains(ReferenceName.PackageName)) )
+			{
+				bool bIsHardReference = HardReferenceNames.Contains(ReferenceName);
+
+				if ( !ExceedsMaxSearchBreadth(NumReferencesMade) )
+				{
+					int32 ThisNodeSizeY = ReferenceName.IsValue() ? 100 : NodeSizeY;
+
+					const int32 RefSizeY = NodeSizes.FindChecked(ReferenceName);
+					FIntPoint RefNodeLoc;
+					RefNodeLoc.X = ReferenceNodeLoc.X;
+					RefNodeLoc.Y = ReferenceNodeLoc.Y + RefSizeY * ThisNodeSizeY * 0.5 - ThisNodeSizeY * 0.5;
+					
+					TArray<FAssetIdentifier> NewIdentifiers;
+					NewIdentifiers.Add(ReferenceName);
+					
+					UEdGraphNode_Reference* ReferenceNode = RecursivelyConstructNodes(bReferencers, RootNode, NewIdentifiers, RefNodeLoc, NodeSizes, PackagesToAssetDataMap, AllowedPackageNames, CurrentDepth + 1, VisitedNames);
+					if (bIsHardReference)
+					{
+						if (bReferencers)
+						{
+							ReferenceNode->GetDependencyPin()->PinType.PinCategory = TEXT("hard");
+						}
+						else
+						{
+							ReferenceNode->GetReferencerPin()->PinType.PinCategory = TEXT("hard"); //-V595
+						}
+					}
+
+					if ( ensure(ReferenceNode) )
+					{
+						if ( bReferencers )
+						{
+							NewNode->AddReferencer( ReferenceNode );
+						}
+						else
+						{
+							ReferenceNode->AddReferencer( NewNode );
+						}
+
+						ReferenceNodeLoc.Y += RefSizeY * ThisNodeSizeY;
+					}
+
+					NumReferencesMade++;
+				}
+				else
+				{
+					NumReferencesExceedingMax++;
+				}
+			}
+		}
+
+		if ( NumReferencesExceedingMax > 0 )
+		{
+			// There are more references than allowed to be displayed. Make a collapsed node.
+			UEdGraphNode_Reference* ReferenceNode = CreateReferenceNode();
+			FIntPoint RefNodeLoc;
+			RefNodeLoc.X = ReferenceNodeLoc.X;
+			RefNodeLoc.Y = ReferenceNodeLoc.Y;
+
+			if ( ensure(ReferenceNode) )
+			{
+				ReferenceNode->SetReferenceNodeCollapsed(RefNodeLoc, NumReferencesExceedingMax);
+
+				if ( bReferencers )
+				{
+					NewNode->AddReferencer( ReferenceNode );
+				}
+				else
+				{
+					ReferenceNode->AddReferencer( NewNode );
+				}
+			}
+		}
+	}
+
+	return NewNode;
+}
+
+const TSharedPtr<FAssetThumbnailPool>& UEdGraph_ReferenceViewer::GetAssetThumbnailPool() const
+{
+	return AssetThumbnailPool;
+}
+
+bool UEdGraph_ReferenceViewer::ExceedsMaxSearchDepth(int32 Depth) const
+{
+	return bLimitSearchDepth && Depth > MaxSearchDepth;
+}
+
+bool UEdGraph_ReferenceViewer::ExceedsMaxSearchBreadth(int32 Breadth) const
+{
+	return bLimitSearchBreadth && Breadth > MaxSearchBreadth;
+}
+
+UEdGraphNode_Reference* UEdGraph_ReferenceViewer::CreateReferenceNode()
+{
+	const bool bSelectNewNode = false;
+	return Cast<UEdGraphNode_Reference>(CreateNode(UEdGraphNode_Reference::StaticClass(), bSelectNewNode));
+}
+
+void UEdGraph_ReferenceViewer::RemoveAllNodes()
+{
+	TArray<UEdGraphNode*> NodesToRemove = Nodes;
+	for (int32 NodeIndex = 0; NodeIndex < NodesToRemove.Num(); ++NodeIndex)
+	{
+		RemoveNode(NodesToRemove[NodeIndex]);
+	}
+}
+
+bool UEdGraph_ReferenceViewer::ShouldFilterByCollection() const
+{
+	return bEnableCollectionFilter && CurrentCollectionFilter != NAME_None;
+}
diff --git a/Engine/Plugins/Editor/AssetManagerEditor/Source/AssetManagerEditor/Private/ReferenceViewer/ReferenceViewerSchema.cpp b/Engine/Plugins/Editor/AssetManagerEditor/Source/AssetManagerEditor/Private/ReferenceViewer/ReferenceViewerSchema.cpp
new file mode 100644
--- /dev/null
+++ b/Engine/Plugins/Editor/AssetManagerEditor/Source/AssetManagerEditor/Private/ReferenceViewer/ReferenceViewerSchema.cpp
@@ -0,0 +1,190 @@
+// Copyright 1998-2017 Epic Games, Inc. All Rights Reserved.
+
+#include "ReferenceViewerSchema.h"
+#include "Textures/SlateIcon.h"
+#include "Misc/Attribute.h"
+#include "Framework/MultiBox/MultiBoxBuilder.h"
+#include "EditorStyleSet.h"
+#include "CollectionManagerTypes.h"
+#include "AssetManagerEditorCommands.h"
+#include "AssetManagerEditorModule.h"
+#include "Toolkits/GlobalEditorCommonCommands.h"
+#include "ConnectionDrawingPolicy.h"
+
+
+static const FLinearColor RiceFlower = FLinearColor(FColor(236, 252, 227));
+static const FLinearColor CannonPink = FLinearColor(FColor(145, 66, 117));
+
+// Overridden connection drawing policy to use less curvy lines between nodes
+class FReferenceViewerConnectionDrawingPolicy : public FConnectionDrawingPolicy
+{
+public:
+	FReferenceViewerConnectionDrawingPolicy(int32 InBackLayerID, int32 InFrontLayerID, float InZoomFactor, const FSlateRect& InClippingRect, FSlateWindowElementList& InDrawElements)
+		: FConnectionDrawingPolicy(InBackLayerID, InFrontLayerID, InZoomFactor, InClippingRect, InDrawElements)
+	{
+	}
+
+	virtual FVector2D ComputeSplineTangent(const FVector2D& Start, const FVector2D& End) const override
+	{
+		const int32 Tension = FMath::Abs<int32>(Start.X - End.X);
+		return Tension * FVector2D(1.0f, 0);
+	}
+
+	virtual void DetermineWiringStyle(UEdGraphPin* OutputPin, UEdGraphPin* InputPin, /*inout*/ FConnectionParams& Params) override
+	{
+		if (OutputPin->PinType.PinCategory == TEXT("hard") || InputPin->PinType.PinCategory == TEXT("hard"))
+		{
+			Params.WireColor = RiceFlower;
+		}
+		else
+		{
+			Params.WireColor = CannonPink;
+		}
+	}
+};
+
+//////////////////////////////////////////////////////////////////////////
+// UReferenceViewerSchema
+
+UReferenceViewerSchema::UReferenceViewerSchema(const FObjectInitializer& ObjectInitializer)
+	: Super(ObjectInitializer)
+{
+}
+
+void UReferenceViewerSchema::GetContextMenuActions(const UEdGraph* CurrentGraph, const UEdGraphNode* InGraphNode, const UEdGraphPin* InGraphPin, class FMenuBuilder* MenuBuilder, bool bIsDebugging) const
+{
+	MenuBuilder->BeginSection(TEXT("Asset"), NSLOCTEXT("ReferenceViewerSchema", "AssetSectionLabel", "Asset"));
+	{
+		MenuBuilder->AddMenuEntry(FGlobalEditorCommonCommands::Get().FindInContentBrowser);
+		MenuBuilder->AddMenuEntry(FAssetManagerEditorCommands::Get().OpenSelectedInAssetEditor);
+	}
+	MenuBuilder->EndSection();
+
+	MenuBuilder->BeginSection(TEXT("Misc"), NSLOCTEXT("ReferenceViewerSchema", "MiscSectionLabel", "Misc"));
+	{
+		MenuBuilder->AddMenuEntry(FAssetManagerEditorCommands::Get().ReCenterGraph);
+		MenuBuilder->AddSubMenu(
+			NSLOCTEXT("ReferenceViewerSchema", "MakeCollectionWithTitle", "Make Collection with"),
+			NSLOCTEXT("ReferenceViewerSchema", "MakeCollectionWithTooltip", "Makes a collection with either the referencers or dependencies of the selected nodes."),
+			FNewMenuDelegate::CreateUObject(this, &UReferenceViewerSchema::GetMakeCollectionWithSubMenu)
+		);
+	}
+	MenuBuilder->EndSection();
+
+	MenuBuilder->BeginSection(TEXT("References"), NSLOCTEXT("ReferenceViewerSchema", "ReferencesSectionLabel", "References"));
+	{
+		MenuBuilder->AddMenuEntry(FAssetManagerEditorCommands::Get().CopyReferencedObjects);
+		MenuBuilder->AddMenuEntry(FAssetManagerEditorCommands::Get().CopyReferencingObjects);
+		MenuBuilder->AddMenuEntry(FAssetManagerEditorCommands::Get().ShowReferencedObjects);
+		MenuBuilder->AddMenuEntry(FAssetManagerEditorCommands::Get().ShowReferencingObjects);
+		MenuBuilder->AddMenuEntry(FAssetManagerEditorCommands::Get().ShowReferenceTree);
+		MenuBuilder->AddMenuEntry(FAssetManagerEditorCommands::Get().ViewSizeMap);
+		MenuBuilder->AddMenuEntry(FAssetManagerEditorCommands::Get().ViewAssetAudit, TEXT("ContextMenu"));
+	}
+	MenuBuilder->EndSection();
+}
+
+FLinearColor UReferenceViewerSchema::GetPinTypeColor(const FEdGraphPinType& PinType) const
+{
+	if (PinType.PinCategory == TEXT("hard"))
+	{
+		return RiceFlower;
+	}
+	else
+	{
+		return CannonPink;
+	}
+}
+
+void UReferenceViewerSchema::BreakPinLinks(UEdGraphPin& TargetPin, bool bSendsNodeNotifcation) const
+{
+	// Don't allow breaking any links
+}
+
+void UReferenceViewerSchema::BreakSinglePinLink(UEdGraphPin* SourcePin, UEdGraphPin* TargetPin) const
+{
+	// Don't allow breaking any links
+}
+
+FPinConnectionResponse UReferenceViewerSchema::MovePinLinks(UEdGraphPin& MoveFromPin, UEdGraphPin& MoveToPin, bool bIsItermediateMove) const
+{
+	// Don't allow moving any links
+	return FPinConnectionResponse(CONNECT_RESPONSE_DISALLOW, FString());
+}
+
+FPinConnectionResponse UReferenceViewerSchema::CopyPinLinks(UEdGraphPin& CopyFromPin, UEdGraphPin& CopyToPin, bool bIsItermediateCopy) const
+{
+	// Don't allow copying any links
+	return FPinConnectionResponse(CONNECT_RESPONSE_DISALLOW, FString());
+}
+
+FConnectionDrawingPolicy* UReferenceViewerSchema::CreateConnectionDrawingPolicy(int32 InBackLayerID, int32 InFrontLayerID, float InZoomFactor, const FSlateRect& InClippingRect, class FSlateWindowElementList& InDrawElements, class UEdGraph* InGraphObj) const
+{
+	return new FReferenceViewerConnectionDrawingPolicy(InBackLayerID, InFrontLayerID, InZoomFactor, InClippingRect, InDrawElements);
+}
+
+void UReferenceViewerSchema::DroppedAssetsOnGraph(const TArray<FAssetData>& Assets, const FVector2D& GraphPosition, UEdGraph* Graph) const
+{
+	TArray<FAssetIdentifier> AssetIdentifiers;
+
+	IAssetManagerEditorModule::ExtractAssetIdentifiersFromAssetDataList(Assets, AssetIdentifiers);
+	IAssetManagerEditorModule::Get().OpenReferenceViewerUI(AssetIdentifiers);
+}
+
+void UReferenceViewerSchema::GetAssetsGraphHoverMessage(const TArray<FAssetData>& Assets, const UEdGraph* HoverGraph, FString& OutTooltipText, bool& OutOkIcon) const
+{
+	OutOkIcon = true;
+}
+
+void UReferenceViewerSchema::GetMakeCollectionWithSubMenu(FMenuBuilder& MenuBuilder)
+{
+	MenuBuilder.AddSubMenu(NSLOCTEXT("ReferenceViewerSchema", "MakeCollectionWithReferencersTitle", "Referencers <-"),
+		NSLOCTEXT("ReferenceViewerSchema", "MakeCollectionWithReferencersTooltip", "Makes a collection with assets one connection to the left of selected nodes."),
+		FNewMenuDelegate::CreateUObject(this, &UReferenceViewerSchema::GetMakeCollectionWithReferencersOrDependenciesSubMenu, true)
+		);
+
+	MenuBuilder.AddSubMenu(NSLOCTEXT("ReferenceViewerSchema", "MakeCollectionWithDependenciesTitle", "Dependencies ->"),
+		NSLOCTEXT("ReferenceViewerSchema", "MakeCollectionWithDependenciesTooltip", "Makes a collection with assets one connection to the right of selected nodes."),
+		FNewMenuDelegate::CreateUObject(this, &UReferenceViewerSchema::GetMakeCollectionWithReferencersOrDependenciesSubMenu, false)
+		);
+}
+
+void UReferenceViewerSchema::GetMakeCollectionWithReferencersOrDependenciesSubMenu(FMenuBuilder& MenuBuilder, bool bReferencers)
+{
+	if (bReferencers)
+	{
+		MenuBuilder.AddMenuEntry(FAssetManagerEditorCommands::Get().MakeLocalCollectionWithReferencers, 
+			NAME_None, TAttribute<FText>(), 
+			ECollectionShareType::GetDescription(ECollectionShareType::CST_Local), 
+			FSlateIcon(FEditorStyle::GetStyleSetName(), ECollectionShareType::GetIconStyleName(ECollectionShareType::CST_Local))
+			);
+		MenuBuilder.AddMenuEntry(FAssetManagerEditorCommands::Get().MakePrivateCollectionWithReferencers,
+			NAME_None, TAttribute<FText>(), 
+			ECollectionShareType::GetDescription(ECollectionShareType::CST_Private), 
+			FSlateIcon(FEditorStyle::GetStyleSetName(), ECollectionShareType::GetIconStyleName(ECollectionShareType::CST_Private))
+			);
+		MenuBuilder.AddMenuEntry(FAssetManagerEditorCommands::Get().MakeSharedCollectionWithReferencers,
+			NAME_None, TAttribute<FText>(), 
+			ECollectionShareType::GetDescription(ECollectionShareType::CST_Shared), 
+			FSlateIcon(FEditorStyle::GetStyleSetName(), ECollectionShareType::GetIconStyleName(ECollectionShareType::CST_Shared))
+			);
+	}
+	else
+	{
+		MenuBuilder.AddMenuEntry(FAssetManagerEditorCommands::Get().MakeLocalCollectionWithDependencies, 
+			NAME_None, TAttribute<FText>(), 
+			ECollectionShareType::GetDescription(ECollectionShareType::CST_Local), 
+			FSlateIcon(FEditorStyle::GetStyleSetName(), ECollectionShareType::GetIconStyleName(ECollectionShareType::CST_Local))
+			);
+		MenuBuilder.AddMenuEntry(FAssetManagerEditorCommands::Get().MakePrivateCollectionWithDependencies,
+			NAME_None, TAttribute<FText>(), 
+			ECollectionShareType::GetDescription(ECollectionShareType::CST_Private), 
+			FSlateIcon(FEditorStyle::GetStyleSetName(), ECollectionShareType::GetIconStyleName(ECollectionShareType::CST_Private))
+			);
+		MenuBuilder.AddMenuEntry(FAssetManagerEditorCommands::Get().MakeSharedCollectionWithDependencies,
+			NAME_None, TAttribute<FText>(), 
+			ECollectionShareType::GetDescription(ECollectionShareType::CST_Shared), 
+			FSlateIcon(FEditorStyle::GetStyleSetName(), ECollectionShareType::GetIconStyleName(ECollectionShareType::CST_Shared))
+			);
+	}
+}
diff --git a/Engine/Plugins/Editor/AssetManagerEditor/Source/AssetManagerEditor/Private/ReferenceViewer/SReferenceViewer.cpp b/Engine/Plugins/Editor/AssetManagerEditor/Source/AssetManagerEditor/Private/ReferenceViewer/SReferenceViewer.cpp
new file mode 100644
--- /dev/null
+++ b/Engine/Plugins/Editor/AssetManagerEditor/Source/AssetManagerEditor/Private/ReferenceViewer/SReferenceViewer.cpp
@@ -0,0 +1,1417 @@
+// Copyright 1998-2017 Epic Games, Inc. All Rights Reserved.
+
+#include "SReferenceViewer.h"
+#include "Widgets/SOverlay.h"
+#include "Engine/GameViewportClient.h"
+#include "Dialogs/Dialogs.h"
+#include "Misc/MessageDialog.h"
+#include "Misc/ScopedSlowTask.h"
+#include "Modules/ModuleManager.h"
+#include "Widgets/SBoxPanel.h"
+#include "Widgets/Layout/SBorder.h"
+#include "Widgets/Images/SImage.h"
+#include "Widgets/Text/STextBlock.h"
+#include "Widgets/Layout/SBox.h"
+#include "Widgets/Input/SEditableTextBox.h"
+#include "Widgets/Input/SButton.h"
+#include "Widgets/Input/SCheckBox.h"
+#include "Widgets/Input/SSpinBox.h"
+#include "EditorStyleSet.h"
+#include "Engine/Selection.h"
+#include "EdGraph_ReferenceViewer.h"
+#include "EdGraphNode_Reference.h"
+#include "ReferenceViewerSchema.h"
+#include "AssetRegistryModule.h"
+#include "ICollectionManager.h"
+#include "CollectionManagerModule.h"
+#include "Editor.h"
+#include "Toolkits/AssetEditorManager.h"
+#include "AssetManagerEditorCommands.h"
+#include "EditorWidgetsModule.h"
+#include "Toolkits/GlobalEditorCommonCommands.h"
+#include "Engine/AssetManager.h"
+#include "Widgets/Input/SComboBox.h"
+#include "HAL/PlatformApplicationMisc.h"
+#include "AssetManagerEditorModule.h"
+
+#include "ObjectTools.h"
+
+#define LOCTEXT_NAMESPACE "ReferenceViewer"
+
+SReferenceViewer::~SReferenceViewer()
+{
+	if (!GExitPurge)
+	{
+		if ( ensure(GraphObj) )
+		{
+			GraphObj->RemoveFromRoot();
+		}		
+	}
+}
+
+void SReferenceViewer::Construct(const FArguments& InArgs)
+{
+	// Create an action list and register commands
+	RegisterActions();
+
+	// Set up the history manager
+	HistoryManager.SetOnApplyHistoryData(FOnApplyHistoryData::CreateSP(this, &SReferenceViewer::OnApplyHistoryData));
+	HistoryManager.SetOnUpdateHistoryData(FOnUpdateHistoryData::CreateSP(this, &SReferenceViewer::OnUpdateHistoryData));
+
+	// Fill out CollectionsComboList
+	{
+		TArray<FName> CollectionNames;
+		FCollectionManagerModule& CollectionManagerModule = FCollectionManagerModule::GetModule();
+		CollectionManagerModule.Get().GetCollectionNames(ECollectionShareType::CST_All, CollectionNames);
+		CollectionNames.Sort([](const FName& A, const FName& B) { return A.Compare(B) < 0; });
+		CollectionsComboList.Add(MakeShareable(new FName(NAME_None)));
+		for (FName CollectionName : CollectionNames)
+		{
+			CollectionsComboList.Add(MakeShareable(new FName(CollectionName)));
+		}
+	}
+
+	// Create the graph
+	GraphObj = NewObject<UEdGraph_ReferenceViewer>();
+	GraphObj->Schema = UReferenceViewerSchema::StaticClass();
+	GraphObj->AddToRoot();
+	GraphObj->SetReferenceViewer(StaticCastSharedRef<SReferenceViewer>(AsShared()));
+
+	SGraphEditor::FGraphEditorEvents GraphEvents;
+	GraphEvents.OnNodeDoubleClicked = FSingleNodeEvent::CreateSP(this, &SReferenceViewer::OnNodeDoubleClicked);
+	GraphEvents.OnCreateActionMenu = SGraphEditor::FOnCreateActionMenu::CreateSP(this, &SReferenceViewer::OnCreateGraphActionMenu);
+
+	// Create the graph editor
+	GraphEditorPtr = SNew(SGraphEditor)
+		.AdditionalCommands(ReferenceViewerActions)
+		.GraphToEdit(GraphObj)
+		.GraphEvents(GraphEvents)
+		.OnNavigateHistoryBack(FSimpleDelegate::CreateSP(this, &SReferenceViewer::GraphNavigateHistoryBack))
+		.OnNavigateHistoryForward(FSimpleDelegate::CreateSP(this, &SReferenceViewer::GraphNavigateHistoryForward));
+
+	FEditorWidgetsModule& EditorWidgetsModule = FModuleManager::LoadModuleChecked<FEditorWidgetsModule>("EditorWidgets");
+	TSharedRef<SWidget> AssetDiscoveryIndicator = EditorWidgetsModule.CreateAssetDiscoveryIndicator(EAssetDiscoveryIndicatorScaleMode::Scale_None, FMargin(16, 8), false);
+
+	static const FName DefaultForegroundName("DefaultForeground");
+
+	ChildSlot
+	[
+		SNew(SVerticalBox)
+
+		// Path and history
+		+SVerticalBox::Slot()
+		.AutoHeight()
+		.Padding( 0, 0, 0, 4 )
+		[
+			SNew(SHorizontalBox)
+
+			// History Back Button
+			+SHorizontalBox::Slot()
+			.AutoWidth()
+			.VAlign(VAlign_Center)
+			.Padding(1,0)
+			[
+				SNew(SButton)
+				.ButtonStyle( FEditorStyle::Get(), "HoverHintOnly" )
+				.ForegroundColor( FEditorStyle::GetSlateColor(DefaultForegroundName) )
+				.ToolTipText( this, &SReferenceViewer::GetHistoryBackTooltip )
+				.ContentPadding( 0 )
+				.OnClicked(this, &SReferenceViewer::BackClicked)
+				.IsEnabled(this, &SReferenceViewer::IsBackEnabled)
+				[
+					SNew(SImage) .Image(FEditorStyle::GetBrush("ContentBrowser.HistoryBack"))
+				]
+			]
+
+			// History Forward Button
+			+SHorizontalBox::Slot()
+			.AutoWidth()
+			.Padding(1,0,3,0)
+			.VAlign(VAlign_Center)
+			[
+				SNew(SButton)
+				.ButtonStyle( FEditorStyle::Get(), "HoverHintOnly" )
+				.ForegroundColor( FEditorStyle::GetSlateColor(DefaultForegroundName) )
+				.ToolTipText( this, &SReferenceViewer::GetHistoryForwardTooltip )
+				.ContentPadding( 0 )
+				.OnClicked(this, &SReferenceViewer::ForwardClicked)
+				.IsEnabled(this, &SReferenceViewer::IsForwardEnabled)
+				[
+					SNew(SImage) .Image(FEditorStyle::GetBrush("ContentBrowser.HistoryForward"))
+				]
+			]
+
+			// Path
+			+SHorizontalBox::Slot()
+			.VAlign(VAlign_Fill)
+			.FillWidth(1.f)
+			[
+				SNew(SBorder)
+				.BorderImage( FEditorStyle::GetBrush("ToolPanel.GroupBorder") )
+				[
+					SNew(SEditableTextBox)
+					.Text(this, &SReferenceViewer::GetAddressBarText)
+					.OnTextCommitted(this, &SReferenceViewer::OnAddressBarTextCommitted)
+					.OnTextChanged(this, &SReferenceViewer::OnAddressBarTextChanged)
+					.SelectAllTextWhenFocused(true)
+					.SelectAllTextOnCommit(true)
+					.Style(FEditorStyle::Get(), "ReferenceViewer.PathText")
+				]
+			]
+		]
+
+		// Graph
+		+SVerticalBox::Slot()
+		.FillHeight(1.f)
+		[
+			SNew(SOverlay)
+
+			+SOverlay::Slot()
+			[
+				GraphEditorPtr.ToSharedRef()
+			]
+
+			+SOverlay::Slot()
+			.HAlign(HAlign_Left)
+			.VAlign(VAlign_Top)
+			.Padding(8)
+			[
+				SNew(SBorder)
+				.BorderImage( FEditorStyle::GetBrush("ToolPanel.GroupBorder") )
+				[
+					SNew(SVerticalBox)
+
+					+SVerticalBox::Slot()
+					.AutoHeight()
+					[
+						SNew(SHorizontalBox)
+
+						+SHorizontalBox::Slot()
+						.VAlign(VAlign_Center)
+						.Padding(2.f)
+						[
+							SNew(STextBlock)
+							.Text(LOCTEXT("SearchDepthLabelText", "Search Depth Limit"))
+						]
+
+						+SHorizontalBox::Slot()
+						.AutoWidth()
+						.VAlign(VAlign_Center)
+						.Padding(2.f)
+						[
+							SNew(SCheckBox)
+							.OnCheckStateChanged( this, &SReferenceViewer::OnSearchDepthEnabledChanged )
+							.IsChecked( this, &SReferenceViewer::IsSearchDepthEnabledChecked )
+						]
+					
+						+SHorizontalBox::Slot()
+						.AutoWidth()
+						.VAlign(VAlign_Center)
+						.Padding(2.f)
+						[
+							SNew(SBox)
+							.WidthOverride(100)
+							[
+								SNew(SSpinBox<int32>)
+								.Value(this, &SReferenceViewer::GetSearchDepthCount)
+								.OnValueChanged(this, &SReferenceViewer::OnSearchDepthCommitted)
+								.MinValue(1)
+								.MaxValue(50)
+								.MaxSliderValue(12)
+							]
+						]
+					]
+
+					+SVerticalBox::Slot()
+					.AutoHeight()
+					[
+						SNew(SHorizontalBox)
+
+						+SHorizontalBox::Slot()
+						.VAlign(VAlign_Center)
+						.Padding(2.f)
+						[
+							SNew(STextBlock)
+							.Text(LOCTEXT("SearchBreadthLabelText", "Search Breadth Limit"))
+						]
+
+						+SHorizontalBox::Slot()
+						.AutoWidth()
+						.VAlign(VAlign_Center)
+						.Padding(2.f)
+						[
+							SNew(SCheckBox)
+							.OnCheckStateChanged( this, &SReferenceViewer::OnSearchBreadthEnabledChanged )
+							.IsChecked( this, &SReferenceViewer::IsSearchBreadthEnabledChecked )
+						]
+					
+						+SHorizontalBox::Slot()
+						.AutoWidth()
+						.VAlign(VAlign_Center)
+						.Padding(2.f)
+						[
+							SNew(SBox)
+							.WidthOverride(100)
+							[
+								SNew(SSpinBox<int32>)
+								.Value(this, &SReferenceViewer::GetSearchBreadthCount)
+								.OnValueChanged(this, &SReferenceViewer::OnSearchBreadthCommitted)
+								.MinValue(1)
+								.MaxValue(1000)
+								.MaxSliderValue(50)
+							]
+						]
+					]
+
+					+SVerticalBox::Slot()
+					.AutoHeight()
+					[
+						SNew(SHorizontalBox)
+
+						+SHorizontalBox::Slot()
+						.VAlign(VAlign_Center)
+						.Padding(2.f)
+						[
+							SNew(STextBlock)
+							.Text(LOCTEXT("CollectionFilter", "Collection Filter"))
+						]
+
+						+SHorizontalBox::Slot()
+						.AutoWidth()
+						.VAlign(VAlign_Center)
+						.Padding(2.f)
+						[
+							SNew(SCheckBox)
+							.OnCheckStateChanged( this, &SReferenceViewer::OnEnableCollectionFilterChanged )
+							.IsChecked( this, &SReferenceViewer::IsEnableCollectionFilterChecked )
+						]
+
+						+SHorizontalBox::Slot()
+						.AutoWidth()
+						.VAlign(VAlign_Center)
+						.Padding(2.f)
+						[
+							SNew(SBox)
+							.WidthOverride(100)
+							[
+								SNew(SComboBox<TSharedPtr<FName>>)
+								.OptionsSource(&CollectionsComboList)
+								.OnGenerateWidget(this, &SReferenceViewer::GenerateCollectionFilterItem)
+								.OnSelectionChanged(this, &SReferenceViewer::HandleCollectionFilterChanged)
+								.ToolTipText(this, &SReferenceViewer::GetCollectionFilterText)
+								[
+									SNew(STextBlock)
+									.Text(this, &SReferenceViewer::GetCollectionFilterText)
+								]
+							]
+						]
+					]
+
+					+SVerticalBox::Slot()
+					.AutoHeight()
+					[
+						SNew(SHorizontalBox)
+
+						+SHorizontalBox::Slot()
+						.VAlign(VAlign_Center)
+						.Padding(2.f)
+						[
+							SNew(STextBlock)
+							.Text(LOCTEXT("ShowHideSoftReferences", "Show Soft References"))
+						]
+
+						+SHorizontalBox::Slot()
+						.AutoWidth()
+						.VAlign(VAlign_Center)
+						.Padding(2.f)
+						[
+							SNew(SCheckBox)
+							.OnCheckStateChanged( this, &SReferenceViewer::OnShowSoftReferencesChanged )
+							.IsChecked( this, &SReferenceViewer::IsShowSoftReferencesChecked )
+						]
+					]
+
+				
+					+ SVerticalBox::Slot()
+					.AutoHeight()
+					[
+						SNew(SHorizontalBox)
+
+						+ SHorizontalBox::Slot()
+						.VAlign(VAlign_Center)
+						.Padding(2.f)
+						[
+							SNew(STextBlock)
+							.Text(LOCTEXT("ShowHideHardReferences", "Show Hard References"))
+						]
+
+						+SHorizontalBox::Slot()
+						.AutoWidth()
+						.VAlign(VAlign_Center)
+						.Padding(2.f)
+						[
+							SNew(SCheckBox)
+							.OnCheckStateChanged(this, &SReferenceViewer::OnShowHardReferencesChanged)
+							.IsChecked(this, &SReferenceViewer::IsShowHardReferencesChecked)
+						]
+					]
+
+					+ SVerticalBox::Slot()
+					.AutoHeight()
+					[
+						SNew(SHorizontalBox)
+						.Visibility(this, &SReferenceViewer::GetManagementReferencesVisibility)
+
+						+ SHorizontalBox::Slot()
+						.VAlign(VAlign_Center)
+						.Padding(2.f)
+						[
+							SNew(STextBlock)
+							.Text(LOCTEXT("ShowHideManagementReferences", "Show Management References"))
+						]
+
+						+SHorizontalBox::Slot()
+						.AutoWidth()
+						.VAlign(VAlign_Center)
+						.Padding(2.f)
+						[
+							SNew(SCheckBox)
+							.OnCheckStateChanged(this, &SReferenceViewer::OnShowManagementReferencesChanged)
+							.IsChecked(this, &SReferenceViewer::IsShowManagementReferencesChecked)
+						]
+					]
+
+					+ SVerticalBox::Slot()
+					.AutoHeight()
+					[
+						SNew(SHorizontalBox)
+
+						+ SHorizontalBox::Slot()
+						.VAlign(VAlign_Center)
+						.Padding(2.f)
+						[
+							SNew(STextBlock)
+							.Text(LOCTEXT("ShowHideSearchableNames", "Show Searchable Names"))
+						]
+
+						+SHorizontalBox::Slot()
+						.AutoWidth()
+						.VAlign(VAlign_Center)
+						.Padding(2.f)
+						[
+							SNew(SCheckBox)
+							.OnCheckStateChanged(this, &SReferenceViewer::OnShowSearchableNamesChanged)
+							.IsChecked(this, &SReferenceViewer::IsShowSearchableNamesChecked)
+						]
+					]
+
+					+ SVerticalBox::Slot()
+					.AutoHeight()
+					[
+						SNew(SHorizontalBox)
+
+						+ SHorizontalBox::Slot()
+						.VAlign(VAlign_Center)
+						.Padding(2.f)
+						[
+							SNew(STextBlock)
+							.Text(LOCTEXT("ShowHideNativePackages", "Show Native Packages"))
+						]
+
+						+SHorizontalBox::Slot()
+						.AutoWidth()
+						.VAlign(VAlign_Center)
+						.Padding(2.f)
+						[
+							SNew(SCheckBox)
+							.OnCheckStateChanged(this, &SReferenceViewer::OnShowNativePackagesChanged)
+							.IsChecked(this, &SReferenceViewer::IsShowNativePackagesChecked)
+						]
+					]
+				]
+			]
+
+			+SOverlay::Slot()
+			.HAlign(HAlign_Center)
+			.VAlign(VAlign_Center)
+			.Padding(FMargin(24, 0, 24, 0))
+			[
+				AssetDiscoveryIndicator
+			]
+		]
+	];
+}
+
+void SReferenceViewer::SetGraphRootIdentifiers(const TArray<FAssetIdentifier>& NewGraphRootIdentifiers)
+{
+	GraphObj->SetGraphRoot(NewGraphRootIdentifiers);
+	
+	RebuildGraph();
+
+	// Zoom once this frame to make sure widgets are visible, then zoom again so size is correct
+	TriggerZoomToFit(0, 0);
+	RegisterActiveTimer(0.1f, FWidgetActiveTimerDelegate::CreateSP(this, &SReferenceViewer::TriggerZoomToFit));
+
+	// Set the initial history data
+	HistoryManager.AddHistoryData();
+}
+
+EActiveTimerReturnType SReferenceViewer::TriggerZoomToFit(double InCurrentTime, float InDeltaTime)
+{
+	if (GraphEditorPtr.IsValid())
+	{
+		GraphEditorPtr->ZoomToFit(false);
+	}
+	return EActiveTimerReturnType::Stop;
+}
+
+void SReferenceViewer::SetCurrentRegistrySource(const FAssetManagerEditorRegistrySource* RegistrySource)
+{
+	RebuildGraph();
+}
+
+void SReferenceViewer::OnNodeDoubleClicked(UEdGraphNode* Node)
+{
+	TSet<UObject*> Nodes;
+	Nodes.Add(Node);
+	ReCenterGraphOnNodes( Nodes );
+}
+
+void SReferenceViewer::RebuildGraph()
+{
+	FAssetRegistryModule& AssetRegistryModule = FModuleManager::LoadModuleChecked<FAssetRegistryModule>(TEXT("AssetRegistry"));
+	if (AssetRegistryModule.Get().IsLoadingAssets())
+	{
+		// We are still discovering assets, listen for the completion delegate before building the graph
+		if (!AssetRegistryModule.Get().OnFilesLoaded().IsBoundToObject(this))
+		{
+			AssetRegistryModule.Get().OnFilesLoaded().AddSP(this, &SReferenceViewer::OnInitialAssetRegistrySearchComplete);
+		}
+	}
+	else
+	{
+		// All assets are already discovered, build the graph now, if we have one
+		if (GraphObj)
+		{
+			GraphObj->RebuildGraph();
+		}
+	}
+}
+
+FActionMenuContent SReferenceViewer::OnCreateGraphActionMenu(UEdGraph* InGraph, const FVector2D& InNodePosition, const TArray<UEdGraphPin*>& InDraggedPins, bool bAutoExpand, SGraphEditor::FActionMenuClosed InOnMenuClosed)
+{
+	// no context menu when not over a node
+	return FActionMenuContent();
+}
+
+bool SReferenceViewer::IsBackEnabled() const
+{
+	return HistoryManager.CanGoBack();
+}
+
+bool SReferenceViewer::IsForwardEnabled() const
+{
+	return HistoryManager.CanGoForward();
+}
+
+FReply SReferenceViewer::BackClicked()
+{
+	HistoryManager.GoBack();
+
+	return FReply::Handled();
+}
+
+FReply SReferenceViewer::ForwardClicked()
+{
+	HistoryManager.GoForward();
+
+	return FReply::Handled();
+}
+
+void SReferenceViewer::GraphNavigateHistoryBack()
+{
+	BackClicked();
+}
+
+void SReferenceViewer::GraphNavigateHistoryForward()
+{
+	ForwardClicked();
+}
+
+FText SReferenceViewer::GetHistoryBackTooltip() const
+{
+	if ( HistoryManager.CanGoBack() )
+	{
+		return FText::Format( LOCTEXT("HistoryBackTooltip", "Back to {0}"), HistoryManager.GetBackDesc() );
+	}
+	return FText::GetEmpty();
+}
+
+FText SReferenceViewer::GetHistoryForwardTooltip() const
+{
+	if ( HistoryManager.CanGoForward() )
+	{
+		return FText::Format( LOCTEXT("HistoryForwardTooltip", "Forward to {0}"), HistoryManager.GetForwardDesc() );
+	}
+	return FText::GetEmpty();
+}
+
+FText SReferenceViewer::GetAddressBarText() const
+{
+	if ( GraphObj )
+	{
+		if (TemporaryPathBeingEdited.IsEmpty())
+		{
+			const TArray<FAssetIdentifier>& CurrentGraphRootPackageNames = GraphObj->GetCurrentGraphRootIdentifiers();
+			if (CurrentGraphRootPackageNames.Num() == 1)
+			{
+				return FText::FromString(CurrentGraphRootPackageNames[0].ToString());
+			}
+			else if (CurrentGraphRootPackageNames.Num() > 1)
+			{
+				return FText::Format(LOCTEXT("AddressBarMultiplePackagesText", "{0} and {1} others"), FText::FromString(CurrentGraphRootPackageNames[0].ToString()), FText::AsNumber(CurrentGraphRootPackageNames.Num()));
+			}
+		}
+		else
+		{
+			return TemporaryPathBeingEdited;
+		}
+	}
+
+	return FText();
+}
+
+void SReferenceViewer::OnAddressBarTextCommitted(const FText& NewText, ETextCommit::Type CommitInfo)
+{
+	if (CommitInfo == ETextCommit::OnEnter)
+	{
+		TArray<FAssetIdentifier> NewPaths;
+		NewPaths.Add(FAssetIdentifier::FromString(NewText.ToString()));
+
+		SetGraphRootIdentifiers(NewPaths);
+	}
+
+	TemporaryPathBeingEdited = FText();
+}
+
+void SReferenceViewer::OnAddressBarTextChanged(const FText& NewText)
+{
+	TemporaryPathBeingEdited = NewText;
+}
+
+void SReferenceViewer::OnApplyHistoryData(const FReferenceViewerHistoryData& History)
+{
+	if ( GraphObj )
+	{
+		GraphObj->SetGraphRoot(History.Identifiers);
+		UEdGraphNode_Reference* NewRootNode = GraphObj->RebuildGraph();
+		
+		if ( NewRootNode && ensure(GraphEditorPtr.IsValid()) )
+		{
+			GraphEditorPtr->SetNodeSelection(NewRootNode, true);
+		}
+	}
+}
+
+void SReferenceViewer::OnUpdateHistoryData(FReferenceViewerHistoryData& HistoryData) const
+{
+	if ( GraphObj )
+	{
+		const TArray<FAssetIdentifier>& CurrentGraphRootIdentifiers = GraphObj->GetCurrentGraphRootIdentifiers();
+		HistoryData.HistoryDesc = GetAddressBarText();
+		HistoryData.Identifiers = CurrentGraphRootIdentifiers;
+	}
+	else
+	{
+		HistoryData.HistoryDesc = FText::GetEmpty();
+		HistoryData.Identifiers.Empty();
+	}
+}
+
+void SReferenceViewer::OnSearchDepthEnabledChanged( ECheckBoxState NewState )
+{
+	if ( GraphObj )
+	{
+		GraphObj->SetSearchDepthLimitEnabled(NewState == ECheckBoxState::Checked);
+		RebuildGraph();
+	}
+}
+
+ECheckBoxState SReferenceViewer::IsSearchDepthEnabledChecked() const
+{
+	if ( GraphObj )
+	{
+		return GraphObj->IsSearchDepthLimited() ? ECheckBoxState::Checked : ECheckBoxState::Unchecked;
+	}
+	else
+	{
+		return ECheckBoxState::Unchecked;
+	}
+}
+
+int32 SReferenceViewer::GetSearchDepthCount() const
+{
+	if ( GraphObj )
+	{
+		return GraphObj->GetSearchDepthLimit();
+	}
+	else
+	{
+		return 0;
+	}
+}
+
+void SReferenceViewer::OnSearchDepthCommitted(int32 NewValue)
+{
+	if ( GraphObj )
+	{
+		GraphObj->SetSearchDepthLimit(NewValue);
+		RebuildGraph();
+	}
+}
+
+void SReferenceViewer::OnSearchBreadthEnabledChanged( ECheckBoxState NewState )
+{
+	if ( GraphObj )
+	{
+		GraphObj->SetSearchBreadthLimitEnabled(NewState == ECheckBoxState::Checked);
+		RebuildGraph();
+	}
+}
+
+ECheckBoxState SReferenceViewer::IsSearchBreadthEnabledChecked() const
+{
+	if ( GraphObj )
+	{
+		return GraphObj->IsSearchBreadthLimited() ? ECheckBoxState::Checked : ECheckBoxState::Unchecked;
+	}
+	else
+	{
+		return ECheckBoxState::Unchecked;
+	}
+}
+
+TSharedRef<SWidget> SReferenceViewer::GenerateCollectionFilterItem(TSharedPtr<FName> InItem)
+{
+	FText ItemAsText = FText::FromName(*InItem);
+	return
+		SNew(SBox)
+		.WidthOverride(300)
+		[
+			SNew(STextBlock)
+			.Text(ItemAsText)
+			.ToolTipText(ItemAsText)
+		];
+}
+
+void SReferenceViewer::OnEnableCollectionFilterChanged(ECheckBoxState NewState)
+{
+	if (GraphObj)
+	{
+		const bool bNewValue = NewState == ECheckBoxState::Checked;
+		const bool bCurrentValue = GraphObj->GetEnableCollectionFilter();
+		if (bCurrentValue != bNewValue)
+		{
+			GraphObj->SetEnableCollectionFilter(NewState == ECheckBoxState::Checked);
+			RebuildGraph();
+		}
+	}
+}
+
+ECheckBoxState SReferenceViewer::IsEnableCollectionFilterChecked() const
+{
+	if (GraphObj)
+	{
+		return GraphObj->GetEnableCollectionFilter() ? ECheckBoxState::Checked : ECheckBoxState::Unchecked;
+	}
+	else
+	{
+		return ECheckBoxState::Unchecked;
+	}
+}
+
+void SReferenceViewer::HandleCollectionFilterChanged(TSharedPtr<FName> Item, ESelectInfo::Type SelectInfo)
+{
+	if (GraphObj)
+	{
+		const FName NewFilter = *Item;
+		const FName CurrentFilter = GraphObj->GetCurrentCollectionFilter();
+		if (CurrentFilter != NewFilter)
+		{
+			if (CurrentFilter == NAME_None)
+			{
+				// Automatically check the box to enable the filter if the previous filter was None
+				GraphObj->SetEnableCollectionFilter(true);
+			}
+
+			GraphObj->SetCurrentCollectionFilter(NewFilter);
+			RebuildGraph();
+		}
+	}
+}
+
+FText SReferenceViewer::GetCollectionFilterText() const
+{
+	return FText::FromName(GraphObj->GetCurrentCollectionFilter());
+}
+
+void SReferenceViewer::OnShowSoftReferencesChanged(ECheckBoxState NewState)
+{
+	if (GraphObj)
+	{
+		GraphObj->SetShowSoftReferencesEnabled(NewState == ECheckBoxState::Checked);
+		RebuildGraph();
+	}
+}
+
+ECheckBoxState SReferenceViewer::IsShowSoftReferencesChecked() const
+{
+	if (GraphObj)
+	{
+		return GraphObj->IsShowSoftReferences()? ECheckBoxState::Checked : ECheckBoxState::Unchecked;
+	}
+	else
+	{
+		return ECheckBoxState::Unchecked;
+	}
+}
+
+void SReferenceViewer::OnShowHardReferencesChanged(ECheckBoxState NewState)
+{
+	if (GraphObj)
+	{
+		GraphObj->SetShowHardReferencesEnabled(NewState == ECheckBoxState::Checked);
+		RebuildGraph();
+	}
+}
+
+
+ECheckBoxState SReferenceViewer::IsShowHardReferencesChecked() const
+{
+	if (GraphObj)
+	{
+		return GraphObj->IsShowHardReferences() ? ECheckBoxState::Checked : ECheckBoxState::Unchecked;
+	}
+	else
+	{
+		return ECheckBoxState::Unchecked;
+	}
+}
+
+EVisibility SReferenceViewer::GetManagementReferencesVisibility() const
+{
+	if (UAssetManager::IsValid())
+	{
+		return EVisibility::SelfHitTestInvisible;
+	}
+	return EVisibility::Collapsed;
+}
+
+void SReferenceViewer::OnShowManagementReferencesChanged(ECheckBoxState NewState)
+{
+	if (GraphObj)
+	{
+		// This can take a few seconds if it isn't ready
+		UAssetManager::Get().UpdateManagementDatabase();
+
+		GraphObj->SetShowManagementReferencesEnabled(NewState == ECheckBoxState::Checked);
+		RebuildGraph();
+	}
+}
+
+ECheckBoxState SReferenceViewer::IsShowManagementReferencesChecked() const
+{
+	if (GraphObj)
+	{
+		return GraphObj->IsShowManagementReferences() ? ECheckBoxState::Checked : ECheckBoxState::Unchecked;
+	}
+	else
+	{
+		return ECheckBoxState::Unchecked;
+	}
+}
+
+void SReferenceViewer::OnShowSearchableNamesChanged(ECheckBoxState NewState)
+{
+	if (GraphObj)
+	{
+		GraphObj->SetShowSearchableNames(NewState == ECheckBoxState::Checked);
+		RebuildGraph();
+	}
+}
+
+ECheckBoxState SReferenceViewer::IsShowSearchableNamesChecked() const
+{
+	if (GraphObj)
+	{
+		return GraphObj->IsShowSearchableNames() ? ECheckBoxState::Checked : ECheckBoxState::Unchecked;
+	}
+	else
+	{
+		return ECheckBoxState::Unchecked;
+	}
+}
+
+void SReferenceViewer::OnShowNativePackagesChanged(ECheckBoxState NewState)
+{
+	if (GraphObj)
+	{
+		GraphObj->SetShowNativePackages(NewState == ECheckBoxState::Checked);
+		RebuildGraph();
+	}
+}
+
+ECheckBoxState SReferenceViewer::IsShowNativePackagesChecked() const
+{
+	if (GraphObj)
+	{
+		return GraphObj->IsShowNativePackages() ? ECheckBoxState::Checked : ECheckBoxState::Unchecked;
+	}
+	else
+	{
+		return ECheckBoxState::Unchecked;
+	}
+}
+
+int32 SReferenceViewer::GetSearchBreadthCount() const
+{
+	if ( GraphObj )
+	{
+		return GraphObj->GetSearchBreadthLimit();
+	}
+	else
+	{
+		return 0;
+	}
+}
+
+void SReferenceViewer::OnSearchBreadthCommitted(int32 NewValue)
+{
+	if ( GraphObj )
+	{
+		GraphObj->SetSearchBreadthLimit(NewValue);
+		RebuildGraph();
+	}
+}
+
+void SReferenceViewer::RegisterActions()
+{
+	ReferenceViewerActions = MakeShareable(new FUICommandList);
+	FAssetManagerEditorCommands::Register();
+
+	ReferenceViewerActions->MapAction(
+		FGlobalEditorCommonCommands::Get().FindInContentBrowser,
+		FExecuteAction::CreateSP(this, &SReferenceViewer::ShowSelectionInContentBrowser),
+		FCanExecuteAction::CreateSP(this, &SReferenceViewer::HasAtLeastOnePackageNodeSelected));
+
+	ReferenceViewerActions->MapAction(
+		FAssetManagerEditorCommands::Get().OpenSelectedInAssetEditor,
+		FExecuteAction::CreateSP(this, &SReferenceViewer::OpenSelectedInAssetEditor),
+		FCanExecuteAction::CreateSP(this, &SReferenceViewer::HasExactlyOnePackageNodeSelected));
+
+	ReferenceViewerActions->MapAction(
+		FAssetManagerEditorCommands::Get().ReCenterGraph,
+		FExecuteAction::CreateSP(this, &SReferenceViewer::ReCenterGraph),
+		FCanExecuteAction());
+
+	ReferenceViewerActions->MapAction(
+		FAssetManagerEditorCommands::Get().CopyReferencedObjects,
+		FExecuteAction::CreateSP(this, &SReferenceViewer::CopyReferencedObjects),
+		FCanExecuteAction::CreateSP(this, &SReferenceViewer::HasAtLeastOnePackageNodeSelected));
+
+	ReferenceViewerActions->MapAction(
+		FAssetManagerEditorCommands::Get().CopyReferencingObjects,
+		FExecuteAction::CreateSP(this, &SReferenceViewer::CopyReferencingObjects),
+		FCanExecuteAction::CreateSP(this, &SReferenceViewer::HasAtLeastOnePackageNodeSelected));
+
+	ReferenceViewerActions->MapAction(
+		FAssetManagerEditorCommands::Get().ShowReferencedObjects,
+		FExecuteAction::CreateSP(this, &SReferenceViewer::ShowReferencedObjects),
+		FCanExecuteAction::CreateSP(this, &SReferenceViewer::HasAtLeastOnePackageNodeSelected));
+
+	ReferenceViewerActions->MapAction(
+		FAssetManagerEditorCommands::Get().ShowReferencingObjects,
+		FExecuteAction::CreateSP(this, &SReferenceViewer::ShowReferencingObjects),
+		FCanExecuteAction::CreateSP(this, &SReferenceViewer::HasAtLeastOnePackageNodeSelected));
+
+	ReferenceViewerActions->MapAction(
+		FAssetManagerEditorCommands::Get().MakeLocalCollectionWithReferencers,
+		FExecuteAction::CreateSP(this, &SReferenceViewer::MakeCollectionWithReferencersOrDependencies, ECollectionShareType::CST_Local, true),
+		FCanExecuteAction::CreateSP(this, &SReferenceViewer::HasExactlyOnePackageNodeSelected));
+	
+	ReferenceViewerActions->MapAction(
+		FAssetManagerEditorCommands::Get().MakePrivateCollectionWithReferencers,
+		FExecuteAction::CreateSP(this, &SReferenceViewer::MakeCollectionWithReferencersOrDependencies, ECollectionShareType::CST_Private, true),
+		FCanExecuteAction::CreateSP(this, &SReferenceViewer::HasExactlyOnePackageNodeSelected));
+	
+	ReferenceViewerActions->MapAction(
+		FAssetManagerEditorCommands::Get().MakeSharedCollectionWithReferencers,
+		FExecuteAction::CreateSP(this, &SReferenceViewer::MakeCollectionWithReferencersOrDependencies, ECollectionShareType::CST_Shared, true),
+		FCanExecuteAction::CreateSP(this, &SReferenceViewer::HasExactlyOnePackageNodeSelected));
+
+	ReferenceViewerActions->MapAction(
+		FAssetManagerEditorCommands::Get().MakeLocalCollectionWithDependencies,
+		FExecuteAction::CreateSP(this, &SReferenceViewer::MakeCollectionWithReferencersOrDependencies, ECollectionShareType::CST_Local, false),
+		FCanExecuteAction::CreateSP(this, &SReferenceViewer::HasExactlyOnePackageNodeSelected));
+
+	ReferenceViewerActions->MapAction(
+		FAssetManagerEditorCommands::Get().MakePrivateCollectionWithDependencies,
+		FExecuteAction::CreateSP(this, &SReferenceViewer::MakeCollectionWithReferencersOrDependencies, ECollectionShareType::CST_Private, false),
+		FCanExecuteAction::CreateSP(this, &SReferenceViewer::HasExactlyOnePackageNodeSelected));
+
+	ReferenceViewerActions->MapAction(
+		FAssetManagerEditorCommands::Get().MakeSharedCollectionWithDependencies,
+		FExecuteAction::CreateSP(this, &SReferenceViewer::MakeCollectionWithReferencersOrDependencies, ECollectionShareType::CST_Shared, false),
+		FCanExecuteAction::CreateSP(this, &SReferenceViewer::HasExactlyOnePackageNodeSelected));
+	
+	ReferenceViewerActions->MapAction(
+		FAssetManagerEditorCommands::Get().ShowReferenceTree,
+		FExecuteAction::CreateSP(this, &SReferenceViewer::ShowReferenceTree),
+		FCanExecuteAction::CreateSP(this, &SReferenceViewer::HasExactlyOnePackageNodeSelected));
+
+	ReferenceViewerActions->MapAction(
+		FAssetManagerEditorCommands::Get().ViewSizeMap,
+		FExecuteAction::CreateSP(this, &SReferenceViewer::ViewSizeMap),
+		FCanExecuteAction());
+
+	ReferenceViewerActions->MapAction(
+		FAssetManagerEditorCommands::Get().ViewAssetAudit,
+		FExecuteAction::CreateSP(this, &SReferenceViewer::ViewAssetAudit),
+		FCanExecuteAction());
+}
+
+void SReferenceViewer::ShowSelectionInContentBrowser()
+{
+	TArray<FAssetData> AssetList;
+
+	// Build up a list of selected assets from the graph selection set
+	TSet<UObject*> SelectedNodes = GraphEditorPtr->GetSelectedNodes();
+	for (FGraphPanelSelectionSet::TConstIterator It(SelectedNodes); It; ++It)
+	{
+		if (UEdGraphNode_Reference* ReferenceNode = Cast<UEdGraphNode_Reference>(*It))
+		{
+			if (ReferenceNode->GetAssetData().IsValid())
+			{
+				AssetList.Add(ReferenceNode->GetAssetData());
+			}
+		}
+	}
+
+	if (AssetList.Num() > 0)
+	{
+		GEditor->SyncBrowserToObjects(AssetList);
+	}
+}
+
+void SReferenceViewer::OpenSelectedInAssetEditor()
+{
+	FAssetRegistryModule& AssetRegistryModule = FModuleManager::LoadModuleChecked<FAssetRegistryModule>(TEXT("AssetRegistry"));
+	TSet<UObject*> SelectedNodes = GraphEditorPtr->GetSelectedNodes();
+	for (FGraphPanelSelectionSet::TConstIterator It(SelectedNodes); It; ++It)
+	{
+		if (UEdGraphNode_Reference* ReferenceNode = Cast<UEdGraphNode_Reference>(*It))
+		{
+			if (!ReferenceNode->IsPackage() && !ReferenceNode->IsCollapsed())
+			{
+				if (AssetRegistryModule.Get().EditSearchableName(ReferenceNode->GetIdentifier()))
+				{
+					// Was handled by callback, otherwise fall back to default behavior
+					return;
+				}
+			}
+		}
+	}
+
+	UObject* SelectedObject = GetObjectFromSingleSelectedNode();
+
+	if( SelectedObject )
+	{
+		FAssetEditorManager::Get().OpenEditorForAsset(SelectedObject);
+	}
+}
+
+void SReferenceViewer::ReCenterGraph()
+{
+	ReCenterGraphOnNodes( GraphEditorPtr->GetSelectedNodes() );
+}
+
+FString SReferenceViewer::GetReferencedObjectsList() const
+{
+	FString ReferencedObjectsList;
+
+	TSet<FName> AllSelectedPackageNames;
+	GetPackageNamesFromSelectedNodes(AllSelectedPackageNames);
+
+	if (AllSelectedPackageNames.Num() > 0)
+	{
+		FAssetRegistryModule& AssetRegistryModule = FModuleManager::LoadModuleChecked<FAssetRegistryModule>(TEXT("AssetRegistry"));
+
+		for (const FName SelectedPackageName : AllSelectedPackageNames)
+		{
+			TArray<FName> HardDependencies;
+			AssetRegistryModule.Get().GetDependencies(SelectedPackageName, HardDependencies, EAssetRegistryDependencyType::Hard);
+			
+			TArray<FName> SoftDependencies;
+			AssetRegistryModule.Get().GetDependencies(SelectedPackageName, SoftDependencies, EAssetRegistryDependencyType::Soft);
+
+			ReferencedObjectsList += FString::Printf(TEXT("[%s - Dependencies]\n"), *SelectedPackageName.ToString());
+			if (HardDependencies.Num() > 0)
+			{
+				ReferencedObjectsList += TEXT("  [HARD]\n");
+				for (const FName HardDependency : HardDependencies)
+				{
+					const FString PackageString = HardDependency.ToString();
+					ReferencedObjectsList += FString::Printf(TEXT("    %s.%s\n"), *PackageString, *FPackageName::GetLongPackageAssetName(PackageString));
+				}
+			}
+			if (SoftDependencies.Num() > 0)
+			{
+				ReferencedObjectsList += TEXT("  [SOFT]\n");
+				for (const FName SoftDependency : SoftDependencies)
+				{
+					const FString PackageString = SoftDependency.ToString();
+					ReferencedObjectsList += FString::Printf(TEXT("    %s.%s\n"), *PackageString, *FPackageName::GetLongPackageAssetName(PackageString));
+				}
+			}
+		}
+	}
+
+	return ReferencedObjectsList;
+}
+
+FString SReferenceViewer::GetReferencingObjectsList() const
+{
+	FString ReferencingObjectsList;
+
+	TSet<FName> AllSelectedPackageNames;
+	GetPackageNamesFromSelectedNodes(AllSelectedPackageNames);
+
+	if (AllSelectedPackageNames.Num() > 0)
+	{
+		FAssetRegistryModule& AssetRegistryModule = FModuleManager::LoadModuleChecked<FAssetRegistryModule>(TEXT("AssetRegistry"));
+
+		for (const FName SelectedPackageName : AllSelectedPackageNames)
+		{
+			TArray<FName> HardDependencies;
+			AssetRegistryModule.Get().GetReferencers(SelectedPackageName, HardDependencies, EAssetRegistryDependencyType::Hard);
+
+			TArray<FName> SoftDependencies;
+			AssetRegistryModule.Get().GetReferencers(SelectedPackageName, SoftDependencies, EAssetRegistryDependencyType::Soft);
+
+			ReferencingObjectsList += FString::Printf(TEXT("[%s - Referencers]\n"), *SelectedPackageName.ToString());
+			if (HardDependencies.Num() > 0)
+			{
+				ReferencingObjectsList += TEXT("  [HARD]\n");
+				for (const FName HardDependency : HardDependencies)
+				{
+					const FString PackageString = HardDependency.ToString();
+					ReferencingObjectsList += FString::Printf(TEXT("    %s.%s\n"), *PackageString, *FPackageName::GetLongPackageAssetName(PackageString));
+				}
+			}
+			if (SoftDependencies.Num() > 0)
+			{
+				ReferencingObjectsList += TEXT("  [SOFT]\n");
+				for (const FName SoftDependency : SoftDependencies)
+				{
+					const FString PackageString = SoftDependency.ToString();
+					ReferencingObjectsList += FString::Printf(TEXT("    %s.%s\n"), *PackageString, *FPackageName::GetLongPackageAssetName(PackageString));
+				}
+			}
+		}
+	}
+
+	return ReferencingObjectsList;
+}
+
+void SReferenceViewer::CopyReferencedObjects()
+{
+	const FString ReferencedObjectsList = GetReferencedObjectsList();
+	FPlatformApplicationMisc::ClipboardCopy(*ReferencedObjectsList);
+}
+
+void SReferenceViewer::CopyReferencingObjects()
+{
+	const FString ReferencingObjectsList = GetReferencingObjectsList();
+	FPlatformApplicationMisc::ClipboardCopy(*ReferencingObjectsList);
+}
+
+void SReferenceViewer::ShowReferencedObjects()
+{
+	const FString ReferencedObjectsList = GetReferencedObjectsList();
+	SGenericDialogWidget::OpenDialog(LOCTEXT("ReferencedObjectsDlgTitle", "Referenced Objects"), SNew(STextBlock).Text(FText::FromString(ReferencedObjectsList)));
+}
+
+void SReferenceViewer::ShowReferencingObjects()
+{	
+	const FString ReferencingObjectsList = GetReferencingObjectsList();
+	SGenericDialogWidget::OpenDialog(LOCTEXT("ReferencingObjectsDlgTitle", "Referencing Objects"), SNew(STextBlock).Text(FText::FromString(ReferencingObjectsList)));
+}
+
+void SReferenceViewer::MakeCollectionWithReferencersOrDependencies(ECollectionShareType::Type ShareType, bool bReferencers)
+{
+	TSet<FName> AllSelectedPackageNames;
+	GetPackageNamesFromSelectedNodes(AllSelectedPackageNames);
+
+	if (AllSelectedPackageNames.Num() > 0)
+	{
+		if (ensure(ShareType != ECollectionShareType::CST_All))
+		{
+			FCollectionManagerModule& CollectionManagerModule = FCollectionManagerModule::GetModule();
+
+			FText CollectionNameAsText;
+			FString FirstAssetName = FPackageName::GetLongPackageAssetName(AllSelectedPackageNames.Array()[0].ToString());
+			if (bReferencers)
+			{
+				if (AllSelectedPackageNames.Num() > 1)
+				{
+					CollectionNameAsText = FText::Format(LOCTEXT("ReferencersForMultipleAssetNames", "{0}AndOthers_Referencers"), FText::FromString(FirstAssetName));
+				}
+				else
+				{
+					CollectionNameAsText = FText::Format(LOCTEXT("ReferencersForSingleAsset", "{0}_Referencers"), FText::FromString(FirstAssetName));
+				}
+			}
+			else
+			{
+				if (AllSelectedPackageNames.Num() > 1)
+				{
+					CollectionNameAsText = FText::Format(LOCTEXT("DependenciesForMultipleAssetNames", "{0}AndOthers_Dependencies"), FText::FromString(FirstAssetName));
+				}
+				else
+				{
+					CollectionNameAsText = FText::Format(LOCTEXT("DependenciesForSingleAsset", "{0}_Dependencies"), FText::FromString(FirstAssetName));
+				}
+			}
+
+			FName CollectionName;
+			CollectionManagerModule.Get().CreateUniqueCollectionName(*CollectionNameAsText.ToString(), ShareType, CollectionName);
+
+			FText ResultsMessage;
+			
+			FAssetRegistryModule& AssetRegistryModule = FModuleManager::LoadModuleChecked<FAssetRegistryModule>(TEXT("AssetRegistry"));
+			TArray<FName> PackageNamesToAddToCollection;
+			if (bReferencers)
+			{
+				for (FName SelectedPackage : AllSelectedPackageNames)
+				{
+					AssetRegistryModule.Get().GetReferencers(SelectedPackage, PackageNamesToAddToCollection);
+				}
+			}
+			else
+			{
+				for (FName SelectedPackage : AllSelectedPackageNames)
+				{
+					AssetRegistryModule.Get().GetDependencies(SelectedPackage, PackageNamesToAddToCollection);
+				}
+			}
+
+			TSet<FName> PackageNameSet;
+			for (FName PackageToAdd : PackageNamesToAddToCollection)
+			{
+				if (!AllSelectedPackageNames.Contains(PackageToAdd))
+				{
+					PackageNameSet.Add(PackageToAdd);
+				}
+			}
+
+			IAssetManagerEditorModule::Get().WriteCollection(CollectionName, ShareType, PackageNameSet.Array(), true);
+		}
+	}
+}
+
+void SReferenceViewer::ShowReferenceTree()
+{
+	UObject* SelectedObject = GetObjectFromSingleSelectedNode();
+
+	if ( SelectedObject )
+	{
+		bool bObjectWasSelected = false;
+		for (FSelectionIterator It(*GEditor->GetSelectedObjects()) ; It; ++It)
+		{
+			if ( (*It) == SelectedObject )
+			{
+				GEditor->GetSelectedObjects()->Deselect( SelectedObject );
+				bObjectWasSelected = true;
+			}
+		}
+
+		ObjectTools::ShowReferenceGraph( SelectedObject );
+
+		if ( bObjectWasSelected )
+		{
+			GEditor->GetSelectedObjects()->Select( SelectedObject );
+		}
+	}
+}
+
+void SReferenceViewer::ViewSizeMap()
+{
+	TArray<FAssetIdentifier> AssetIdentifiers;
+	TSet<UObject*> SelectedNodes = GraphEditorPtr->GetSelectedNodes();
+	for (UObject* Node : SelectedNodes)
+	{
+		UEdGraphNode_Reference* ReferenceNode = Cast<UEdGraphNode_Reference>(Node);
+		if (ReferenceNode)
+		{
+			ReferenceNode->GetAllIdentifiers(AssetIdentifiers);
+		}
+	}
+
+	if (AssetIdentifiers.Num() > 0)
+	{
+		IAssetManagerEditorModule::Get().OpenSizeMapUI(AssetIdentifiers);
+	}
+}
+
+void SReferenceViewer::ViewAssetAudit()
+{
+	TSet<FName> SelectedAssetPackageNames;
+	GetPackageNamesFromSelectedNodes(SelectedAssetPackageNames);
+
+	if (SelectedAssetPackageNames.Num() > 0)
+	{
+		IAssetManagerEditorModule::Get().OpenAssetAuditUI(SelectedAssetPackageNames.Array());
+	}
+}
+
+void SReferenceViewer::ReCenterGraphOnNodes(const TSet<UObject*>& Nodes)
+{
+	TArray<FAssetIdentifier> NewGraphRootNames;
+	FIntPoint TotalNodePos(EForceInit::ForceInitToZero);
+	for ( auto NodeIt = Nodes.CreateConstIterator(); NodeIt; ++NodeIt )
+	{
+		UEdGraphNode_Reference* ReferenceNode = Cast<UEdGraphNode_Reference>(*NodeIt);
+		if ( ReferenceNode )
+		{
+			ReferenceNode->GetAllIdentifiers(NewGraphRootNames);
+			TotalNodePos.X += ReferenceNode->NodePosX;
+			TotalNodePos.Y += ReferenceNode->NodePosY;
+		}
+	}
+
+	if ( NewGraphRootNames.Num() > 0 )
+	{
+		const FIntPoint AverageNodePos = TotalNodePos / NewGraphRootNames.Num();
+		GraphObj->SetGraphRoot(NewGraphRootNames, AverageNodePos);
+		UEdGraphNode_Reference* NewRootNode = GraphObj->RebuildGraph();
+
+		if ( NewRootNode && ensure(GraphEditorPtr.IsValid()) )
+		{
+			GraphEditorPtr->ClearSelectionSet();
+			GraphEditorPtr->SetNodeSelection(NewRootNode, true);
+		}
+
+		// Set the initial history data
+		HistoryManager.AddHistoryData();
+	}
+}
+
+UObject* SReferenceViewer::GetObjectFromSingleSelectedNode() const
+{
+	UObject* ReturnObject = nullptr;
+
+	TSet<UObject*> SelectedNodes = GraphEditorPtr->GetSelectedNodes();
+	if ( ensure(SelectedNodes.Num()) == 1 )
+	{
+		UEdGraphNode_Reference* ReferenceNode = Cast<UEdGraphNode_Reference>(SelectedNodes.Array()[0]);
+		if ( ReferenceNode )
+		{
+			const FAssetData& AssetData = ReferenceNode->GetAssetData();
+			if (AssetData.IsAssetLoaded())
+			{
+				ReturnObject = AssetData.GetAsset();
+			}
+			else
+			{
+				FScopedSlowTask SlowTask(0, LOCTEXT("LoadingSelectedObject", "Loading selection..."));
+				SlowTask.MakeDialog();
+				ReturnObject = AssetData.GetAsset();
+			}
+		}
+	}
+
+	return ReturnObject;
+}
+
+void SReferenceViewer::GetPackageNamesFromSelectedNodes(TSet<FName>& OutNames) const
+{
+	TSet<UObject*> SelectedNodes = GraphEditorPtr->GetSelectedNodes();
+	for (UObject* Node : SelectedNodes)
+	{
+		UEdGraphNode_Reference* ReferenceNode = Cast<UEdGraphNode_Reference>(Node);
+		if (ReferenceNode)
+		{
+			TArray<FName> NodePackageNames;
+			ReferenceNode->GetAllPackageNames(NodePackageNames);
+			OutNames.Append(NodePackageNames);
+		}
+	}
+}
+
+bool SReferenceViewer::HasExactlyOneNodeSelected() const
+{
+	if ( GraphEditorPtr.IsValid() )
+	{
+		return GraphEditorPtr->GetSelectedNodes().Num() == 1;
+	}
+	
+	return false;
+}
+
+bool SReferenceViewer::HasExactlyOnePackageNodeSelected() const
+{
+	if (GraphEditorPtr.IsValid())
+	{
+		if (GraphEditorPtr->GetSelectedNodes().Num() != 1)
+		{
+			return false;
+		}
+
+		TSet<UObject*> SelectedNodes = GraphEditorPtr->GetSelectedNodes();
+		for (UObject* Node : SelectedNodes)
+		{
+			UEdGraphNode_Reference* ReferenceNode = Cast<UEdGraphNode_Reference>(Node);
+			if (ReferenceNode)
+			{
+				if (ReferenceNode->IsPackage())
+				{
+					return true;
+				}
+			}
+			return false;
+		}
+	}
+
+	return false;
+}
+
+bool SReferenceViewer::HasAtLeastOnePackageNodeSelected() const
+{
+	if ( GraphEditorPtr.IsValid() )
+	{
+		TSet<UObject*> SelectedNodes = GraphEditorPtr->GetSelectedNodes();
+		for (UObject* Node : SelectedNodes)
+		{
+			UEdGraphNode_Reference* ReferenceNode = Cast<UEdGraphNode_Reference>(Node);
+			if (ReferenceNode)
+			{
+				if (ReferenceNode->IsPackage())
+				{
+					return true;
+				}
+			}
+		}
+	}
+	
+	return false;
+}
+
+void SReferenceViewer::OnInitialAssetRegistrySearchComplete()
+{
+	if ( GraphObj )
+	{
+		GraphObj->RebuildGraph();
+	}
+}
+
+#undef LOCTEXT_NAMESPACE
diff --git a/Engine/Plugins/Editor/AssetManagerEditor/Source/AssetManagerEditor/Private/ReferenceViewer/SReferenceViewer.h b/Engine/Plugins/Editor/AssetManagerEditor/Source/AssetManagerEditor/Private/ReferenceViewer/SReferenceViewer.h
new file mode 100644
--- /dev/null
+++ b/Engine/Plugins/Editor/AssetManagerEditor/Source/AssetManagerEditor/Private/ReferenceViewer/SReferenceViewer.h
@@ -0,0 +1,163 @@
+// Copyright 1998-2017 Epic Games, Inc. All Rights Reserved.
+
+#pragma once
+
+#include "CoreMinimal.h"
+#include "Input/Reply.h"
+#include "Widgets/DeclarativeSyntaxSupport.h"
+#include "Widgets/SCompoundWidget.h"
+#include "GraphEditor.h"
+#include "AssetData.h"
+#include "HistoryManager.h"
+#include "CollectionManagerTypes.h"
+#include "AssetManagerEditorModule.h"
+
+class UEdGraph;
+class UEdGraph_ReferenceViewer;
+
+/**
+ * 
+ */
+class SReferenceViewer : public SCompoundWidget
+{
+public:
+	SLATE_BEGIN_ARGS( SReferenceViewer ){}
+
+	SLATE_END_ARGS()
+
+	~SReferenceViewer();
+
+	/** Constructs this widget with InArgs */
+	void Construct( const FArguments& InArgs );
+
+	/** Sets a new root package name */
+	void SetGraphRootIdentifiers(const TArray<FAssetIdentifier>& NewGraphRootIdentifiers);
+
+	/** Gets graph editor */
+	TSharedPtr<SGraphEditor> GetGraphEditor() const { return GraphEditorPtr; }
+
+	/** Called when the current registry source changes */
+	void SetCurrentRegistrySource(const FAssetManagerEditorRegistrySource* RegistrySource);
+
+private:
+
+	/** Call after a structural change is made that causes the graph to be recreated */
+	void RebuildGraph();
+
+	/** Called to create context menu when right-clicking on graph */
+	FActionMenuContent OnCreateGraphActionMenu(UEdGraph* InGraph, const FVector2D& InNodePosition, const TArray<UEdGraphPin*>& InDraggedPins, bool bAutoExpand, SGraphEditor::FActionMenuClosed InOnMenuClosed);
+
+	/** Called when a node is double clicked */
+	void OnNodeDoubleClicked(class UEdGraphNode* Node);
+
+	/** True if the user may use the history back button */
+	bool IsBackEnabled() const;
+
+	/** True if the user may use the history forward button */
+	bool IsForwardEnabled() const;
+
+	/** Handler for clicking the history back button */
+	FReply BackClicked();
+
+	/** Handler for clicking the history forward button */
+	FReply ForwardClicked();
+
+	/** Handler for when the graph panel tells us to go back in history (like using the mouse thumb button) */
+	void GraphNavigateHistoryBack();
+
+	/** Handler for when the graph panel tells us to go forward in history (like using the mouse thumb button) */
+	void GraphNavigateHistoryForward();
+
+	/** Gets the tool tip text for the history back button */
+	FText GetHistoryBackTooltip() const;
+
+	/** Gets the tool tip text for the history forward button */
+	FText GetHistoryForwardTooltip() const;
+
+	/** Gets the text to be displayed in the address bar */
+	FText GetAddressBarText() const;
+
+	/** Called when the path is being edited */
+	void OnAddressBarTextChanged(const FText& NewText);
+
+	/** Sets the new path for the viewer */
+	void OnAddressBarTextCommitted(const FText& NewText, ETextCommit::Type CommitInfo);
+
+	void OnApplyHistoryData(const FReferenceViewerHistoryData& History);
+
+	void OnUpdateHistoryData(FReferenceViewerHistoryData& HistoryData) const;
+	
+	void OnSearchDepthEnabledChanged( ECheckBoxState NewState );
+	ECheckBoxState IsSearchDepthEnabledChecked() const;
+	int32 GetSearchDepthCount() const;
+	void OnSearchDepthCommitted(int32 NewValue);
+
+	void OnSearchBreadthEnabledChanged( ECheckBoxState NewState );
+	ECheckBoxState IsSearchBreadthEnabledChecked() const;
+
+	void OnEnableCollectionFilterChanged(ECheckBoxState NewState);
+	ECheckBoxState IsEnableCollectionFilterChecked() const;
+	TSharedRef<SWidget> GenerateCollectionFilterItem(TSharedPtr<FName> InItem);
+	void HandleCollectionFilterChanged(TSharedPtr<FName> Item, ESelectInfo::Type SelectInfo);
+	FText GetCollectionFilterText() const;
+
+	void OnShowSoftReferencesChanged( ECheckBoxState NewState );
+	ECheckBoxState IsShowSoftReferencesChecked() const;
+	void OnShowHardReferencesChanged(ECheckBoxState NewState);
+	ECheckBoxState IsShowHardReferencesChecked() const;
+
+	EVisibility GetManagementReferencesVisibility() const;
+	void OnShowManagementReferencesChanged(ECheckBoxState NewState);
+	ECheckBoxState IsShowManagementReferencesChecked() const;
+
+	void OnShowSearchableNamesChanged(ECheckBoxState NewState);
+	ECheckBoxState IsShowSearchableNamesChecked() const;
+	void OnShowNativePackagesChanged(ECheckBoxState NewState);
+	ECheckBoxState IsShowNativePackagesChecked() const;
+
+	int32 GetSearchBreadthCount() const;
+	void OnSearchBreadthCommitted(int32 NewValue);
+
+	void RegisterActions();
+	void ShowSelectionInContentBrowser();
+	void OpenSelectedInAssetEditor();
+	void ReCenterGraph();
+	void CopyReferencedObjects();
+	void CopyReferencingObjects();
+	void ShowReferencedObjects();
+	void ShowReferencingObjects();
+	void MakeCollectionWithReferencersOrDependencies(ECollectionShareType::Type ShareType, bool bReferencers);
+	void ShowReferenceTree();
+	void ViewSizeMap();
+	void ViewAssetAudit();
+
+	void ReCenterGraphOnNodes(const TSet<UObject*>& Nodes);
+
+	FString GetReferencedObjectsList() const;
+	FString GetReferencingObjectsList() const;
+
+	UObject* GetObjectFromSingleSelectedNode() const;
+	void GetPackageNamesFromSelectedNodes(TSet<FName>& OutNames) const;
+	bool HasExactlyOneNodeSelected() const;
+	bool HasExactlyOnePackageNodeSelected() const;
+	bool HasAtLeastOnePackageNodeSelected() const;
+
+	void OnInitialAssetRegistrySearchComplete();
+	EActiveTimerReturnType TriggerZoomToFit(double InCurrentTime, float InDeltaTime);
+private:
+
+	/** The manager that keeps track of history data for this browser */
+	FReferenceViewerHistoryManager HistoryManager;
+
+	TSharedPtr<SGraphEditor> GraphEditorPtr;
+
+	TSharedPtr<FUICommandList> ReferenceViewerActions;
+
+	UEdGraph_ReferenceViewer* GraphObj;
+
+	/** The temporary copy of the path text when it is actively being edited. */
+	FText TemporaryPathBeingEdited;
+
+	/** List of collection filter options */
+	TArray<TSharedPtr<FName>> CollectionsComboList;
+};
