Copying WEX-Staging @ (WEX/Main @ 3740665) to //UE4/Main

diff --git a/Engine/Plugins/Editor/AssetManagerEditor/Source/AssetManagerEditor/Private/SAssetAuditBrowser.cpp b/Engine/Plugins/Editor/AssetManagerEditor/Source/AssetManagerEditor/Private/SAssetAuditBrowser.cpp
--- a/Engine/Plugins/Editor/AssetManagerEditor/Source/AssetManagerEditor/Private/SAssetAuditBrowser.cpp
+++ b/Engine/Plugins/Editor/AssetManagerEditor/Source/AssetManagerEditor/Private/SAssetAuditBrowser.cpp
@@ -18,7 +18,6 @@
 #include "ARFilter.h"
 #include "ClassViewerModule.h"
 #include "ClassViewerFilter.h"
-
 #include "IContentBrowserSingleton.h"
 #include "ContentBrowserModule.h"
 #include "AssetRegistryModule.h"
@@ -27,9 +26,10 @@
 #include "Slate/SceneViewport.h"
 #include "ObjectEditorUtils.h"
 #include "Engine/AssetManager.h"
+#include "AssetManagerEditorCommands.h"
 #include "Widgets/Input/SComboBox.h"
-#include "ReferenceViewer.h"
 #include "SlateApplication.h"
+#include "DragAndDrop/AssetDragDropOp.h"
 
 #define LOCTEXT_NAMESPACE "AssetManagementBrowser"
 
@@ -44,26 +44,17 @@ TSharedPtr<SWidget> SAssetAuditBrowser::OnGetAssetContextMenu(const TArray<FAsse
 {
 	FMenuBuilder MenuBuilder(/*bInShouldCloseWindowAfterMenuSelection=*/ true, Commands);
 
-	MenuBuilder.BeginSection("AssetOptions", LOCTEXT("OptionsHeading", "Options") );
+	MenuBuilder.BeginSection(TEXT("Asset"), NSLOCTEXT("ReferenceViewerSchema", "AssetSectionLabel", "Asset"));
 	{
-		MenuBuilder.AddMenuEntry(
-			LOCTEXT("EditAsset", "Edit..."),
-			LOCTEXT("EditAssetTooltip", "Opens the selected asset(s) for edit."),
-			FSlateIcon(FEditorStyle::GetStyleSetName(), "ContentBrowser.AssetActions.Edit"),
-			FUIAction(
-				FExecuteAction::CreateSP(this, &SAssetAuditBrowser::EditSelectedAssets, SelectedAssets),
-				FCanExecuteAction::CreateSP(this, &SAssetAuditBrowser::IsAnythingSelected)
-			)
-		);
-
 		MenuBuilder.AddMenuEntry(FGlobalEditorCommonCommands::Get().FindInContentBrowser);
+		MenuBuilder.AddMenuEntry(FAssetManagerEditorCommands::Get().OpenSelectedInAssetEditor);
 
 		MenuBuilder.AddMenuEntry(
 			LOCTEXT("Load", "Load..."),
 			LOCTEXT("LoadTooltip", "Loads selected assets into memory."),
 			FSlateIcon(FEditorStyle::GetStyleSetName(), "LevelEditor.OpenLevel"),
 			FUIAction(
-				FExecuteAction::CreateSP(this, &SAssetAuditBrowser::LoadSelectedAssets, SelectedAssets),
+				FExecuteAction::CreateSP(this, &SAssetAuditBrowser::LoadSelectedAssets),
 				FCanExecuteAction::CreateSP(this, &SAssetAuditBrowser::IsAnythingSelected)
 			)
 		);
@@ -73,20 +64,17 @@ TSharedPtr<SWidget> SAssetAuditBrowser::OnGetAssetContextMenu(const TArray<FAsse
 			LOCTEXT("SaveSelectedAssets_ToolTip", "Save the selected assets."),
 			FSlateIcon(FEditorStyle::GetStyleSetName(), "LevelEditor.Save"),
 			FUIAction(
-				FExecuteAction::CreateSP(this, &SAssetAuditBrowser::SaveSelectedAssets, SelectedAssets),
-				FCanExecuteAction::CreateSP(this, &SAssetAuditBrowser::CanSaveSelectedAssets, SelectedAssets)
+				FExecuteAction::CreateSP(this, &SAssetAuditBrowser::SaveSelectedAssets),
+				FCanExecuteAction::CreateSP(this, &SAssetAuditBrowser::IsAnythingSelectedAndLoaded)
 			)
 		);
+	}
+	MenuBuilder.EndSection();
 
-		MenuBuilder.AddMenuEntry(
-			LOCTEXT("ReferenceViewer", "Find References..."),
-			LOCTEXT("ReferenceViewerTooltip", "Shows a graph of references for this asset."),
-			FSlateIcon(),
-			FUIAction(
-				FExecuteAction::CreateSP(this, &SAssetAuditBrowser::FindReferencesForSelectedAssets, SelectedAssets),
-				FCanExecuteAction::CreateSP(this, &SAssetAuditBrowser::IsAnythingSelected)
-			)
-		);
+	MenuBuilder.BeginSection(TEXT("References"), NSLOCTEXT("ReferenceViewerSchema", "ReferencesSectionLabel", "References"));
+	{
+		MenuBuilder.AddMenuEntry(FAssetManagerEditorCommands::Get().ViewReferences);
+		MenuBuilder.AddMenuEntry(FAssetManagerEditorCommands::Get().ViewSizeMap);
 	}
 	MenuBuilder.EndSection();
 
@@ -109,11 +97,21 @@ bool SAssetAuditBrowser::IsAnythingSelected() const
 	return CurrentSelection.Num() > 0;
 }
 
+bool SAssetAuditBrowser::IsAnythingSelectedAndLoaded() const
+{
+	TArray<FAssetData> Assets = GetCurrentSelectionDelegate.Execute();
+	TArray<UPackage*> Packages;
+	GetSelectedPackages(Assets, Packages);
+
+	// Don't offer option if none of the packages are loaded
+	return Packages.Num() > 0;
+}
+
 void SAssetAuditBrowser::GetSelectedPackages(const TArray<FAssetData>& Assets, TArray<UPackage*>& OutPackages) const
 {
-	for (int32 AssetIdx = 0; AssetIdx < Assets.Num(); ++AssetIdx)
+	for (const FAssetData& AssetData : Assets)
 	{
-		UPackage* Package = FindPackage(NULL, *Assets[AssetIdx].PackageName.ToString());
+		UPackage* Package = FindPackage(nullptr, *AssetData.PackageName.ToString());
 
 		if ( Package )
 		{
@@ -122,8 +120,9 @@ void SAssetAuditBrowser::GetSelectedPackages(const TArray<FAssetData>& Assets, T
 	}
 }
 
-void SAssetAuditBrowser::EditSelectedAssets(TArray<FAssetData> Assets) const
+void SAssetAuditBrowser::EditSelectedAssets() const
 {
+	TArray<FAssetData> Assets = GetCurrentSelectionDelegate.Execute();
 	TArray<FName> AssetNames;
 
 	for (FAssetData& AssetData : Assets)
@@ -136,11 +135,15 @@ void SAssetAuditBrowser::EditSelectedAssets(TArray<FAssetData> Assets) const
 
 void SAssetAuditBrowser::OnRequestOpenAsset(const FAssetData& AssetData) const
 {
-	EditSelectedAssets(TArray<FAssetData>{AssetData});
+	TArray<FName> AssetNames;
+	AssetNames.Add(AssetData.ObjectPath);
+
+	FAssetEditorManager::Get().OpenEditorsForAssets(AssetNames);
 }
 
-void SAssetAuditBrowser::SaveSelectedAssets(TArray<FAssetData> Assets) const
+void SAssetAuditBrowser::SaveSelectedAssets() const
 {
+	TArray<FAssetData> Assets = GetCurrentSelectionDelegate.Execute();
 	TArray<UPackage*> PackagesToSave;
 	GetSelectedPackages(Assets, PackagesToSave);
 
@@ -149,30 +152,33 @@ void SAssetAuditBrowser::SaveSelectedAssets(TArray<FAssetData> Assets) const
 	const FEditorFileUtils::EPromptReturnCode Return = FEditorFileUtils::PromptForCheckoutAndSave(PackagesToSave, bCheckDirty, bPromptToSave);
 }
 
-bool SAssetAuditBrowser::CanSaveSelectedAssets(TArray<FAssetData> Assets) const
+void SAssetAuditBrowser::FindReferencesForSelectedAssets() const
 {
-	TArray<UPackage*> Packages;
-	GetSelectedPackages(Assets, Packages);
-	// Don't offer save option if none of the packages are loaded
-	return Packages.Num() > 0;
-}
+	TArray<FAssetData> Assets = GetCurrentSelectionDelegate.Execute();
+	TArray<FAssetIdentifier> AssetIdentifiers;
+	IAssetManagerEditorModule::ExtractAssetIdentifiersFromAssetDataList(Assets, AssetIdentifiers);
 
-void SAssetAuditBrowser::FindReferencesForSelectedAssets(TArray<FAssetData> Assets) const
-{
-	TArray<FName> PackageNames;
-	for (FAssetData& AssetData : Assets)
+	if (AssetIdentifiers.Num() > 0)
 	{
-		PackageNames.Add(AssetData.PackageName);
+		IAssetManagerEditorModule::Get().OpenReferenceViewerUI(AssetIdentifiers);
 	}
+}
+
+void SAssetAuditBrowser::ShowSizeMapForSelectedAssets() const
+{
+	TArray<FAssetData> Assets = GetCurrentSelectionDelegate.Execute();
+	TArray<FAssetIdentifier> AssetIdentifiers;
+	IAssetManagerEditorModule::ExtractAssetIdentifiersFromAssetDataList(Assets, AssetIdentifiers);
 
-	if (PackageNames.Num() > 0)
+	if (AssetIdentifiers.Num() > 0)
 	{
-		IReferenceViewerModule::Get().InvokeReferenceViewerTab(PackageNames);
+		IAssetManagerEditorModule::Get().OpenSizeMapUI(AssetIdentifiers);
 	}
 }
 
-void SAssetAuditBrowser::LoadSelectedAssets(TArray<FAssetData> Assets) const
+void SAssetAuditBrowser::LoadSelectedAssets() const
 {
+	TArray<FAssetData> Assets = GetCurrentSelectionDelegate.Execute();
 	for (FAssetData& AssetData : Assets)
 	{
 		AssetData.GetAsset();
@@ -184,9 +190,24 @@ bool SAssetAuditBrowser::CanShowColumnForAssetRegistryTag(FName AssetType, FName
 	return !AssetRegistryTagsToIgnore.Contains(TagName);
 }
 
-FString SAssetAuditBrowser::GetValueForCustomColumn(FAssetData& AssetData, FName ColumnName) const
+FString SAssetAuditBrowser::GetStringValueForCustomColumn(FAssetData& AssetData, FName ColumnName) const
 {
-	return EditorModule->GetValueForCustomColumn(AssetData, ColumnName, CurrentTargetPlatform, CurrentPlatformState);
+	FString OutValue;
+	EditorModule->GetStringValueForCustomColumn(AssetData, ColumnName, OutValue);
+	return OutValue;
+}
+
+void SAssetAuditBrowser::OnGetCustomSourceAssets(const FARFilter& Filter, TArray<FAssetData>& OutAssets) const
+{
+	if (Filter.PackageNames.Contains(IAssetManagerEditorModule::ChunkFakeAssetDataPackageName) && CurrentRegistrySource)
+	{
+		IAssetManagerEditorModule::Get().GetCurrentRegistrySource(true);
+
+		for (const TPair<int32, FAssetManagerChunkInfo>& Pair : CurrentRegistrySource->ChunkAssignments)
+		{
+			OutAssets.Add(IAssetManagerEditorModule::CreateFakeAssetDataFromChunkId(Pair.Key));
+		}
+	}
 }
 
 void SAssetAuditBrowser::Construct(const FArguments& InArgs)
@@ -204,21 +225,15 @@ void SAssetAuditBrowser::Construct(const FArguments& InArgs)
 
 	AssetManager->UpdateManagementDatabase();
 
-	PlatformComboList.Add(MakeShared<FString>(TEXT("Editor")));
-	PlatformList.Add(nullptr);
+	TArray<const FAssetManagerEditorRegistrySource*> AvailableSources;
+	ManagerEditorModule.GetAvailableRegistrySources(AvailableSources);
 
-	TArray<ITargetPlatform*> ValidPlatforms;
-	ManagerEditorModule.GetAvailableTargetPlatforms(ValidPlatforms);
-
-	for (ITargetPlatform* ValidPlatform : ValidPlatforms)
+	for (const FAssetManagerEditorRegistrySource* ValidSource : AvailableSources)
 	{
-		PlatformList.Add(ValidPlatform);
-		PlatformComboList.Add(MakeShared<FString>(ValidPlatform->PlatformName()));
+		SourceComboList.Add(MakeShared<FString>(ValidSource->SourceName));
 	}
 
-	CurrentPlatformString = TEXT("Editor");
-	CurrentTargetPlatform = nullptr;
-	CurrentPlatformState = nullptr;
+	CurrentRegistrySource = ManagerEditorModule.GetCurrentRegistrySource(false);
 
 	Commands = MakeShareable(new FUICommandList());
 	Commands->MapAction(FGlobalEditorCommonCommands::Get().FindInContentBrowser, FUIAction(
@@ -226,6 +241,21 @@ void SAssetAuditBrowser::Construct(const FArguments& InArgs)
 		FCanExecuteAction::CreateSP(this, &SAssetAuditBrowser::IsAnythingSelected)
 		));
 
+	Commands->MapAction(FAssetManagerEditorCommands::Get().OpenSelectedInAssetEditor, FUIAction(
+		FExecuteAction::CreateSP(this, &SAssetAuditBrowser::EditSelectedAssets),
+		FCanExecuteAction::CreateSP(this, &SAssetAuditBrowser::IsAnythingSelected)
+	));
+	
+	Commands->MapAction(FAssetManagerEditorCommands::Get().ViewReferences, FUIAction(
+		FExecuteAction::CreateSP(this, &SAssetAuditBrowser::FindReferencesForSelectedAssets),
+		FCanExecuteAction::CreateSP(this, &SAssetAuditBrowser::IsAnythingSelected)
+	));
+
+	Commands->MapAction(FAssetManagerEditorCommands::Get().ViewSizeMap, FUIAction(
+		FExecuteAction::CreateSP(this, &SAssetAuditBrowser::ShowSizeMapForSelectedAssets),
+		FCanExecuteAction::CreateSP(this, &SAssetAuditBrowser::IsAnythingSelected)
+	));
+
 	CurrentAssetHistoryIndex = 0;
 	AssetHistory.Add(TSet<FName>());
 
@@ -242,6 +272,7 @@ void SAssetAuditBrowser::Construct(const FArguments& InArgs)
 	Config.OnAssetDoubleClicked = FOnAssetDoubleClicked::CreateSP(this, &SAssetAuditBrowser::OnRequestOpenAsset);
 	Config.OnGetAssetContextMenu = FOnGetAssetContextMenu::CreateSP(this, &SAssetAuditBrowser::OnGetAssetContextMenu);
 	Config.OnAssetTagWantsToBeDisplayed = FOnShouldDisplayAssetTag::CreateSP(this, &SAssetAuditBrowser::CanShowColumnForAssetRegistryTag);
+	Config.OnGetCustomSourceAssets = FOnGetCustomSourceAssets::CreateSP(this, &SAssetAuditBrowser::OnGetCustomSourceAssets);
 	Config.SyncToAssetsDelegates.Add(&SyncToAssetsDelegate);
 	Config.OnShouldFilterAsset = FOnShouldFilterAsset::CreateSP(this, &SAssetAuditBrowser::HandleFilterAsset);
 	Config.GetCurrentSelectionDelegates.Add(&GetCurrentSelectionDelegate);
@@ -256,15 +287,15 @@ void SAssetAuditBrowser::Construct(const FArguments& InArgs)
 	Config.HiddenColumnNames.Add(TEXT("Path"));
 
 	// Add custom columns
-	Config.CustomColumns.Emplace(FPrimaryAssetId::PrimaryAssetTypeTag, LOCTEXT("AssetType", "Primary Type"), LOCTEXT("AssetTypeTooltip", "Primary Asset Type of this asset, if set"), UObject::FAssetRegistryTag::TT_Alphabetical, FOnGetCustomAssetColumnData::CreateSP(this, &SAssetAuditBrowser::GetValueForCustomColumn));
-	Config.CustomColumns.Emplace(FPrimaryAssetId::PrimaryAssetNameTag, LOCTEXT("AssetName", "Primary Name"), LOCTEXT("AssetNameTooltip", "Primary Asset Name of this asset, if set"), UObject::FAssetRegistryTag::TT_Alphabetical, FOnGetCustomAssetColumnData::CreateSP(this, &SAssetAuditBrowser::GetValueForCustomColumn));
-	Config.CustomColumns.Emplace(IAssetManagerEditorModule::ManagedResourceSizeName, LOCTEXT("ManagedResourceSize", "Memory Kb"), LOCTEXT("ManagedResourceSizeTooltip", "Memory used by both this asset and any other assets it manages, in kilobytes"), UObject::FAssetRegistryTag::TT_Numerical, FOnGetCustomAssetColumnData::CreateSP(this, &SAssetAuditBrowser::GetValueForCustomColumn));
-	Config.CustomColumns.Emplace(IAssetManagerEditorModule::ResourceSizeName, LOCTEXT("ResourceSize", "Exclusive Memory Kb"), LOCTEXT("ResourceSizeTooltip", "Memory used exclusively by this asset, in kilobytes"), UObject::FAssetRegistryTag::TT_Numerical, FOnGetCustomAssetColumnData::CreateSP(this, &SAssetAuditBrowser::GetValueForCustomColumn));
-	Config.CustomColumns.Emplace(IAssetManagerEditorModule::ManagedDiskSizeName, LOCTEXT("ManagedDiskSize", "Disk Kb"), LOCTEXT("ManagedDiskSizeTooltip", "Total disk space used by both this and all managed assets, in kilobytes"), UObject::FAssetRegistryTag::TT_Numerical, FOnGetCustomAssetColumnData::CreateSP(this, &SAssetAuditBrowser::GetValueForCustomColumn));
-	Config.CustomColumns.Emplace(IAssetManagerEditorModule::DiskSizeName, LOCTEXT("DiskSize", "Exclusive Disk Kb"), LOCTEXT("DiskSizeTooltip", "Size of saved file on disk for only this asset, in kilobytes"), UObject::FAssetRegistryTag::TT_Numerical, FOnGetCustomAssetColumnData::CreateSP(this, &SAssetAuditBrowser::GetValueForCustomColumn));
-	Config.CustomColumns.Emplace(IAssetManagerEditorModule::TotalUsageName, LOCTEXT("TotalUsage", "Total Usage"), LOCTEXT("TotalUsageTooltip", "Weighted count of Primary Assets that use this, higher usage means it's more likely to be in memory at runtime"), UObject::FAssetRegistryTag::TT_Numerical, FOnGetCustomAssetColumnData::CreateSP(this, &SAssetAuditBrowser::GetValueForCustomColumn));
-	Config.CustomColumns.Emplace(IAssetManagerEditorModule::CookRuleName, LOCTEXT("CookRule", "Cook Rule"), LOCTEXT("CookRuleTooltip", "Rather this asset will be cooked or not"), UObject::FAssetRegistryTag::TT_Alphabetical, FOnGetCustomAssetColumnData::CreateSP(this, &SAssetAuditBrowser::GetValueForCustomColumn));
-	Config.CustomColumns.Emplace(IAssetManagerEditorModule::ChunksName, LOCTEXT("Chunks", "Chunks"), LOCTEXT("ChunksTooltip", "List of chunks this will be added to when cooked"), UObject::FAssetRegistryTag::TT_Alphabetical, FOnGetCustomAssetColumnData::CreateSP(this, &SAssetAuditBrowser::GetValueForCustomColumn));
+	Config.CustomColumns.Emplace(FPrimaryAssetId::PrimaryAssetTypeTag, LOCTEXT("AssetType", "Primary Type"), LOCTEXT("AssetTypeTooltip", "Primary Asset Type of this asset, if set"), UObject::FAssetRegistryTag::TT_Alphabetical, FOnGetCustomAssetColumnData::CreateSP(this, &SAssetAuditBrowser::GetStringValueForCustomColumn));
+	Config.CustomColumns.Emplace(FPrimaryAssetId::PrimaryAssetNameTag, LOCTEXT("AssetName", "Primary Name"), LOCTEXT("AssetNameTooltip", "Primary Asset Name of this asset, if set"), UObject::FAssetRegistryTag::TT_Alphabetical, FOnGetCustomAssetColumnData::CreateSP(this, &SAssetAuditBrowser::GetStringValueForCustomColumn));
+	Config.CustomColumns.Emplace(IAssetManagerEditorModule::ManagedResourceSizeName, LOCTEXT("ManagedResourceSize", "Memory Kb"), LOCTEXT("ManagedResourceSizeTooltip", "Memory used by both this asset and any other assets it manages, in kilobytes"), UObject::FAssetRegistryTag::TT_Numerical, FOnGetCustomAssetColumnData::CreateSP(this, &SAssetAuditBrowser::GetStringValueForCustomColumn));
+	Config.CustomColumns.Emplace(IAssetManagerEditorModule::ResourceSizeName, LOCTEXT("ResourceSize", "Exclusive Memory Kb"), LOCTEXT("ResourceSizeTooltip", "Memory used exclusively by this asset, in kilobytes"), UObject::FAssetRegistryTag::TT_Numerical, FOnGetCustomAssetColumnData::CreateSP(this, &SAssetAuditBrowser::GetStringValueForCustomColumn));
+	Config.CustomColumns.Emplace(IAssetManagerEditorModule::ManagedDiskSizeName, LOCTEXT("ManagedDiskSize", "Disk Kb"), LOCTEXT("ManagedDiskSizeTooltip", "Total disk space used by both this and all managed assets, in kilobytes"), UObject::FAssetRegistryTag::TT_Numerical, FOnGetCustomAssetColumnData::CreateSP(this, &SAssetAuditBrowser::GetStringValueForCustomColumn));
+	Config.CustomColumns.Emplace(IAssetManagerEditorModule::DiskSizeName, LOCTEXT("DiskSize", "Exclusive Disk Kb"), LOCTEXT("DiskSizeTooltip", "Size of saved file on disk for only this asset, in kilobytes"), UObject::FAssetRegistryTag::TT_Numerical, FOnGetCustomAssetColumnData::CreateSP(this, &SAssetAuditBrowser::GetStringValueForCustomColumn));
+	Config.CustomColumns.Emplace(IAssetManagerEditorModule::TotalUsageName, LOCTEXT("TotalUsage", "Total Usage"), LOCTEXT("TotalUsageTooltip", "Weighted count of Primary Assets that use this, higher usage means it's more likely to be in memory at runtime"), UObject::FAssetRegistryTag::TT_Numerical, FOnGetCustomAssetColumnData::CreateSP(this, &SAssetAuditBrowser::GetStringValueForCustomColumn));
+	Config.CustomColumns.Emplace(IAssetManagerEditorModule::CookRuleName, LOCTEXT("CookRule", "Cook Rule"), LOCTEXT("CookRuleTooltip", "Rather this asset will be cooked or not"), UObject::FAssetRegistryTag::TT_Alphabetical, FOnGetCustomAssetColumnData::CreateSP(this, &SAssetAuditBrowser::GetStringValueForCustomColumn));
+	Config.CustomColumns.Emplace(IAssetManagerEditorModule::ChunksName, LOCTEXT("Chunks", "Chunks"), LOCTEXT("ChunksTooltip", "List of chunks this will be added to when cooked"), UObject::FAssetRegistryTag::TT_Alphabetical, FOnGetCustomAssetColumnData::CreateSP(this, &SAssetAuditBrowser::GetStringValueForCustomColumn));
 	
 	// Ignore these tags as we added them as custom columns
 	AssetRegistryTagsToIgnore.Add(FPrimaryAssetId::PrimaryAssetTypeTag);
@@ -371,7 +402,7 @@ void SAssetAuditBrowser::Construct(const FArguments& InArgs)
 						IAssetManagerEditorModule::MakePrimaryAssetTypeSelector(
 							FOnGetPrimaryAssetDisplayText::CreateLambda([] { return LOCTEXT("AddAssetsOfType", "Add Primary Asset Type"); }),
 							FOnSetPrimaryAssetType::CreateSP(this, &SAssetAuditBrowser::AddAssetsOfType),
-							false)
+							false, true)
 					]
 					+SHorizontalBox::Slot()
 					.AutoWidth()
@@ -397,6 +428,15 @@ void SAssetAuditBrowser::Construct(const FArguments& InArgs)
 					]
 					+SHorizontalBox::Slot()
 					.AutoWidth()
+					[
+						SNew(SButton)
+						.HAlign(HAlign_Center)
+						.VAlign(VAlign_Center)
+						.Text(LOCTEXT("AddChunks", "Add Chunks"))
+						.OnClicked(this, &SAssetAuditBrowser::AddChunks)
+					]
+					+SHorizontalBox::Slot()
+					.AutoWidth()
 					[
 						SNew(SButton)
 						.HAlign(HAlign_Center)
@@ -424,12 +464,12 @@ void SAssetAuditBrowser::Construct(const FArguments& InArgs)
 						.HAlign(HAlign_Right)
 						[
 							SNew(SComboBox<TSharedPtr<FString>>)							
-							.OptionsSource(&PlatformComboList)
-							.OnGenerateWidget(this, &SAssetAuditBrowser::GeneratePlatformComboItem)
-							.OnSelectionChanged(this, &SAssetAuditBrowser::HandlePlatformComboChanged)
+							.OptionsSource(&SourceComboList)
+							.OnGenerateWidget(this, &SAssetAuditBrowser::GenerateSourceComboItem)
+							.OnSelectionChanged(this, &SAssetAuditBrowser::HandleSourceComboChanged)
 							[
 								SNew(STextBlock)
-								.Text(this, &SAssetAuditBrowser::GetPlatformComboText)
+								.Text(this, &SAssetAuditBrowser::GetSourceComboText)
 							]
 						]
 					]
@@ -461,6 +501,38 @@ FReply SAssetAuditBrowser::OnKeyDown(const FGeometry& MyGeometry, const FKeyEven
 	return FReply::Unhandled();
 }
 
+FReply SAssetAuditBrowser::OnDragOver(const FGeometry& MyGeometry, const FDragDropEvent& DragDropEvent)
+{
+	bool bIsDragSupported = false;
+
+	TSharedPtr<FDragDropOperation> Operation = DragDropEvent.GetOperation();
+	if (Operation.IsValid() && Operation->IsOfType<FAssetDragDropOp>())
+	{
+		bIsDragSupported = true;
+	}
+
+	return bIsDragSupported ? FReply::Handled() : FReply::Unhandled();
+}
+
+FReply SAssetAuditBrowser::OnDrop(const FGeometry& MyGeometry, const FDragDropEvent& DragDropEvent)
+{
+	bool bWasDropHandled = false;
+	TSharedPtr<FDragDropOperation> Operation = DragDropEvent.GetOperation();
+
+	if (Operation.IsValid())
+	{
+		if (Operation->IsOfType<FAssetDragDropOp>())
+		{
+			const FAssetDragDropOp& DragDropOp = *StaticCastSharedPtr<FAssetDragDropOp>(Operation);
+
+			AddAssetsToList(DragDropOp.GetAssets(), false);
+			bWasDropHandled = true;
+		}
+	}
+
+	return bWasDropHandled ? FReply::Handled() : FReply::Unhandled();
+}
+
 void SAssetAuditBrowser::AddAssetsToList(const TArray<FSoftObjectPath>& AssetsToView, bool bReplaceExisting)
 {
 	TArray<FName> AssetNames;
@@ -617,21 +689,44 @@ FReply SAssetAuditBrowser::ClearAssets()
 	return FReply::Handled();
 }
 
-FReply SAssetAuditBrowser::RefreshAssets()
+FReply SAssetAuditBrowser::AddChunks()
 {
-	AssetManager->UpdateManagementDatabase(true);
-	RefreshAssetView();
+	TArray<FName> PackageNames;
+	PackageNames.Add(IAssetManagerEditorModule::ChunkFakeAssetDataPackageName);
+
+	AddAssetsToList(PackageNames, false);
 
 	return FReply::Handled();
 }
 
+FReply SAssetAuditBrowser::RefreshAssets()
+{
+	// This will end up refreshing the UI indirectly
+	IAssetManagerEditorModule::Get().RefreshRegistryData();
+	
+	return FReply::Handled();
+}
+
 void SAssetAuditBrowser::AddAssetsOfType(FPrimaryAssetType AssetType)
 {
 	if (AssetType.IsValid())
 	{
 		TArray<FSoftObjectPath> AssetArray;
-		
-		AssetManager->GetPrimaryAssetPathList(AssetType, AssetArray);
+
+		if (AssetType == IAssetManagerEditorModule::AllPrimaryAssetTypes)
+		{
+			TArray<FPrimaryAssetTypeInfo> TypeList;
+			AssetManager->GetPrimaryAssetTypeInfoList(TypeList);
+
+			for (const FPrimaryAssetTypeInfo& TypeInfo : TypeList)
+			{
+				AssetManager->GetPrimaryAssetPathList(TypeInfo.PrimaryAssetType, AssetArray);
+			}
+		}
+		else
+		{
+			AssetManager->GetPrimaryAssetPathList(AssetType, AssetArray);
+		}
 
 		AddAssetsToList(AssetArray, false);
 	}
@@ -695,47 +790,43 @@ EVisibility SAssetAuditBrowser::GetHistoryVisibility() const
 
 bool SAssetAuditBrowser::HandleFilterAsset(const FAssetData& InAssetData) const
 {
-	if (CurrentPlatformState)
+	if (InAssetData.PackagePath == IAssetManagerEditorModule::PrimaryAssetFakeAssetDataPackagePath)
 	{
-		const FAssetPackageData* FoundPackageData = CurrentPlatformState->GetAssetPackageData(InAssetData.PackageName);
-
-		// Check if it was actually cooked for this platform, negative size means it was not cooked
-		if (!FoundPackageData || FoundPackageData->DiskSize < 0)
-		{
-			return true;
-		}
+		return false;
 	}
 
-	return false;
+	return !EditorModule->IsPackageInCurrentRegistrySource(InAssetData.PackageName);
 }
 
-TSharedRef<SWidget> SAssetAuditBrowser::GeneratePlatformComboItem(TSharedPtr<FString> InItem)
+TSharedRef<SWidget> SAssetAuditBrowser::GenerateSourceComboItem(TSharedPtr<FString> InItem)
 {
-	return SNew(STextBlock)
-		.Text(FText::FromString(*InItem));
+	return SNew(STextBlock).Text(FText::FromString(*InItem));
 }
 
-void SAssetAuditBrowser::HandlePlatformComboChanged(TSharedPtr<FString> Item, ESelectInfo::Type SelectInfo)
+void SAssetAuditBrowser::HandleSourceComboChanged(TSharedPtr<FString> Item, ESelectInfo::Type SelectInfo)
 {
-	for (int32 i = 0; i < PlatformComboList.Num(); i++)
+	if (Item.IsValid())
 	{
-		if (Item == PlatformComboList[i])
-		{
-			CurrentPlatformString = *Item.Get();
-			CurrentTargetPlatform = PlatformList[i];
-
-			// Cache registry, may be null
-			CurrentPlatformState = EditorModule->GetAssetRegistryStateForTargetPlatform(CurrentTargetPlatform);
+		// This will call back into this widget to refresh it
+		EditorModule->SetCurrentRegistrySource(*Item);
+	}
+}
 
-			// Refresh
-			RefreshAssetView();
-		}
+FText SAssetAuditBrowser::GetSourceComboText() const
+{
+	if (CurrentRegistrySource)
+	{
+		return FText::FromString(CurrentRegistrySource->SourceName);
 	}
+
+	return FText::GetEmpty();
 }
 
-FText SAssetAuditBrowser::GetPlatformComboText() const
+void SAssetAuditBrowser::SetCurrentRegistrySource(const FAssetManagerEditorRegistrySource* RegistrySource)
 {
-	return FText::FromString(CurrentPlatformString);
+	CurrentRegistrySource = RegistrySource;
+
+	RefreshAssetView();
 }
 
 #undef LOCTEXT_NAMESPACE
diff --git a/Engine/Plugins/Editor/AssetManagerEditor/Source/AssetManagerEditor/Private/SAssetAuditBrowser.h b/Engine/Plugins/Editor/AssetManagerEditor/Source/AssetManagerEditor/Private/SAssetAuditBrowser.h
--- a/Engine/Plugins/Editor/AssetManagerEditor/Source/AssetManagerEditor/Private/SAssetAuditBrowser.h
+++ b/Engine/Plugins/Editor/AssetManagerEditor/Source/AssetManagerEditor/Private/SAssetAuditBrowser.h
@@ -35,28 +35,36 @@ public:
 	void AddAssetsToList(const TArray<FSoftObjectPath>& AssetsToView, bool bReplaceExisting);
 	void AddAssetsToList(const TArray<FName>& PackageNamesToView, bool bReplaceExisting);
 
+	/** Called when the current registry source changes */
+	void SetCurrentRegistrySource(const FAssetManagerEditorRegistrySource* RegistrySource);
+
 protected:
 	virtual FReply OnKeyDown(const FGeometry& MyGeometry, const FKeyEvent& InKeyEvent) override;
+	virtual FReply OnDragOver(const FGeometry& MyGeometry, const FDragDropEvent& DragDropEvent) override;
+	virtual FReply OnDrop(const FGeometry& MyGeometry, const FDragDropEvent& DragDropEvent) override;
 
 	/** Delegate that handles creation of context menu */
 	TSharedPtr<SWidget> OnGetAssetContextMenu(const TArray<FAssetData>& SelectedAssets);
 
-	/** Delegate to handle "Find in Content Browser" context menu option */
+	/** Delegates for context menu */
 	void FindInContentBrowser() const;
 	bool IsAnythingSelected() const;
+	bool IsAnythingSelectedAndLoaded() const;
 	void OnRequestOpenAsset(const FAssetData& AssetData) const;
 
-	void EditSelectedAssets(TArray<FAssetData> Assets) const;
-	void SaveSelectedAssets(TArray<FAssetData> Assets) const;
-	bool CanSaveSelectedAssets(TArray<FAssetData> Assets) const;
-	void FindReferencesForSelectedAssets(TArray<FAssetData> Assets) const;
-	void LoadSelectedAssets(TArray<FAssetData> Assets) const;
+	void EditSelectedAssets() const;
+	void SaveSelectedAssets() const;
+	void FindReferencesForSelectedAssets() const;
+	void ShowSizeMapForSelectedAssets() const;
+	void LoadSelectedAssets() const;
 
 	bool CanShowColumnForAssetRegistryTag(FName AssetType, FName TagName) const;
-	FString GetValueForCustomColumn(FAssetData& AssetData, FName ColumnName) const;
-	
+	FString GetStringValueForCustomColumn(FAssetData& AssetData, FName ColumnName) const;
+	void OnGetCustomSourceAssets(const FARFilter& Filter, TArray<FAssetData>& OutAssets) const;
+
 	/** Populate supplied OutPackages with the packages for the supplied Assets array */
 	void GetSelectedPackages(const TArray<FAssetData>& Assets, TArray<UPackage*>& OutPackages) const;
+	void GetSelectedAssetIdentifiers(const TArray<FAssetData>& Assets, TArray<FAssetIdentifier>& OutAssetIdentifiers) const;
 
 	/** Single step forward in history */
 	FReply OnGoForwardInHistory();
@@ -97,6 +105,7 @@ protected:
 
 	/** Button callbacks */
 	FReply ClearAssets();
+	FReply AddChunks();
 	FReply RefreshAssets();
 
 	/** Adds all assets of type */
@@ -112,9 +121,9 @@ protected:
 	TSharedRef<SWidget> CreateClassPicker();
 
 	/** Platform combo box */
-	TSharedRef<SWidget> GeneratePlatformComboItem(TSharedPtr<FString> InItem);
-	void HandlePlatformComboChanged(TSharedPtr<FString> Item, ESelectInfo::Type SelectInfo);
-	FText GetPlatformComboText() const;
+	TSharedRef<SWidget> GenerateSourceComboItem(TSharedPtr<FString> InItem);
+	void HandleSourceComboChanged(TSharedPtr<FString> Item, ESelectInfo::Type SelectInfo);
+	FText GetSourceComboText() const;
 
 	/** Refresh the asset view with a new filter */
 	void RefreshAssetView();
@@ -131,16 +140,11 @@ protected:
 	/** Current position in the asset history */
 	int32 CurrentAssetHistoryIndex;
 
-	/** List of valid platforms */
-	TArray<TSharedPtr<FString>> PlatformComboList;
-	TArray<ITargetPlatform*> PlatformList;
-
-	/** Current platform string */
-	FString CurrentPlatformString;
+	/** List of valid registry sources */
+	TArray<TSharedPtr<FString>> SourceComboList;
 
 	/** Current TargetPlatform and registry state, may be null! */
-	ITargetPlatform* CurrentTargetPlatform;
-	FAssetRegistryState* CurrentPlatformState;
+	const FAssetManagerEditorRegistrySource* CurrentRegistrySource;
 
 	/** Delegates to interact with asset view */
 	FSyncToAssetsDelegate SyncToAssetsDelegate;
diff --git a/Engine/Plugins/Editor/AssetManagerEditor/Source/AssetManagerEditor/Public/AssetManagerEditorModule.h b/Engine/Plugins/Editor/AssetManagerEditor/Source/AssetManagerEditor/Public/AssetManagerEditorModule.h
--- a/Engine/Plugins/Editor/AssetManagerEditor/Source/AssetManagerEditor/Public/AssetManagerEditorModule.h
+++ b/Engine/Plugins/Editor/AssetManagerEditor/Source/AssetManagerEditor/Public/AssetManagerEditorModule.h
@@ -8,6 +8,8 @@
 #include "UObject/PrimaryAssetId.h"
 #include "Interfaces/ITargetPlatform.h"
 #include "AssetRegistryState.h"
+#include "Engine/AssetManagerTypes.h"
+#include "CollectionManagerTypes.h"
 
 DECLARE_LOG_CATEGORY_EXTERN(LogAssetManagerEditor, Log, All);
 
@@ -19,6 +21,53 @@ class SToolTip;
 class SWidget;
 struct FAssetData;
 
+/** Struct containing the information about currently investigated asset context */
+struct ASSETMANAGEREDITOR_API FAssetManagerEditorRegistrySource
+{
+	/** Display name of the source, equal to target Platform, Editor, or Custom */
+	FString SourceName;
+
+	/** Filename registry was loaded from */
+	FString SourceFilename;
+
+	/** Target platform for this state, may be null */
+	ITargetPlatform* TargetPlatform;
+
+	/** Raw asset registry state, if bIsEditor is true this points to the real editor asset registry */
+	const FAssetRegistryState* RegistryState;
+
+	/** If true, this is the editor  */
+	uint8 bIsEditor : 1;
+
+	/** If true, management data has been initialized */
+	uint8 bManagementDataInitialized : 1;
+
+	/** Map of chunk information, from chunk id to primary assets/specific assets that are explicitly assigned to that chunk */
+	TMap<int32, FAssetManagerChunkInfo> ChunkAssignments;
+
+	/** If SourceName matches this, this is the editor */
+	static const FString EditorSourceName;
+
+	/** If SourceName matches this, this is the custom source loaded from an arbitrary path */
+	static const FString CustomSourceName;
+
+	FAssetManagerEditorRegistrySource()
+		: TargetPlatform(nullptr)
+		, RegistryState(nullptr)
+		, bIsEditor(false)
+		, bManagementDataInitialized(false)
+	{
+	}
+
+	~FAssetManagerEditorRegistrySource()
+	{
+		if (RegistryState && !bIsEditor)
+		{
+			delete RegistryState;
+		}
+	}
+};
+
 /**
  * The Asset Manager Editor module handles creating UI for asset management and exposes several commands
  */
@@ -46,18 +95,53 @@ public:
 		return FModuleManager::Get().IsModuleLoaded( "AssetManagerEditor" );
 	}
 
-	/** Creates a simple version of a Primary Asset Type selector, not bound to a PropertyHandle */
-	static TSharedRef<SWidget> MakePrimaryAssetTypeSelector(FOnGetPrimaryAssetDisplayText OnGetDisplayText, FOnSetPrimaryAssetType OnSetType, bool bAllowClear = true);
+	/** 
+	 * Creates a simple version of a Primary Asset Type selector, not bound to a PropertyHandle 
+	 * @param OnGetDisplayText Delegate that returns the text to display in body of combo box
+	 * @param OnSetType Delegate called when type is changed
+	 * @param bAllowClear If true, add None option to top
+	 * @param bAlowAll If true, add All Types option to bottom, returns AllPrimaryAssetTypes if selected
+	 */
+	static TSharedRef<SWidget> MakePrimaryAssetTypeSelector(FOnGetPrimaryAssetDisplayText OnGetDisplayText, FOnSetPrimaryAssetType OnSetType, bool bAllowClear = true, bool bAllowAll = false);
 
-	/** Creates a simple version of a Primary Asset Id selector, not bound to a PropertyHandle */
+	/** 
+	 * Creates a simple version of a Primary Asset Id selector, not bound to a PropertyHandle
+	 * @param OnGetDisplayText Delegate that returns the text to display in body of combo box
+	 * @param OnSetId Delegate called when id is changed
+	 * @param bAllowClear If true, add None option to top
+	 */
 	static TSharedRef<SWidget> MakePrimaryAssetIdSelector(FOnGetPrimaryAssetDisplayText OnGetDisplayText, FOnSetPrimaryAssetId OnSetId, bool bAllowClear = true, TArray<FPrimaryAssetType> AllowedTypes = TArray<FPrimaryAssetType>());
 
 	/** Called to get list of valid primary asset types */
-	static void GeneratePrimaryAssetTypeComboBoxStrings(TArray< TSharedPtr<FString> >& OutComboBoxStrings, TArray<TSharedPtr<SToolTip>>& OutToolTips, TArray<bool>& OutRestrictedItems, bool bAllowClear);
+	static void GeneratePrimaryAssetTypeComboBoxStrings(TArray< TSharedPtr<FString> >& OutComboBoxStrings, TArray<TSharedPtr<SToolTip>>& OutToolTips, TArray<bool>& OutRestrictedItems, bool bAllowClear, bool bAllowAll);
 
 	/** Called for asset picker to see rather to show asset */
 	static bool OnShouldFilterPrimaryAsset(const FAssetData& InAssetData, TArray<FPrimaryAssetType> AllowedTypes);
 
+	/** Creates fake AssetData to represent chunks */
+	static FAssetData CreateFakeAssetDataFromChunkId(int32 ChunkID);
+
+	/** Returns >= 0 if this is a fake chunk AssetData, otherwise INDEX_NONE */
+	static int32 ExtractChunkIdFromFakeAssetData(const FAssetData& InAssetData);
+
+	/** Creates fake AssetData to represent a primary asset */
+	static FAssetData CreateFakeAssetDataFromPrimaryAssetId(const FPrimaryAssetId& PrimaryAssetId);
+
+	/** Returns valid PrimaryAssetId if this is a fake AssetData for a primary asset */
+	static FPrimaryAssetId ExtractPrimaryAssetIdFromFakeAssetData(const FAssetData& InAssetData);
+
+	/** Parses a list of AssetData and extracts AssetIdentifiers, handles the fake asset data */
+	static void ExtractAssetIdentifiersFromAssetDataList(const TArray<FAssetData>& AssetDataList, TArray<FAssetIdentifier>& OutAssetIdentifiers);
+
+	/** Fake package path used by placeholder asset data, has primary asset ID added to it to set package name */
+	static const FName PrimaryAssetFakeAssetDataPackagePath;
+
+	/** Specific package path inside PrimaryAssetFakeAssetDataPackagePath representing chunks */
+	static const FName ChunkFakeAssetDataPackageName;
+	
+	/** Fake type returned from type combo box to mean all types */
+	static const FPrimaryAssetType AllPrimaryAssetTypes;
+
 	/** Custom column names */
 	static const FName ResourceSizeName;
 	static const FName DiskSizeName;
@@ -67,12 +151,45 @@ public:
 	static const FName CookRuleName;
 	static const FName ChunksName;
 
-	/** Returns the value of a "virtual" column for an asset data, this will query the AssetManager for you and takes current platform into account */
-	virtual FString GetValueForCustomColumn(FAssetData& AssetData, FName ColumnName, ITargetPlatform* TargetPlatform, const FAssetRegistryState* PlatformState) = 0;
+	/** Gets the value of a "virtual" column for an asset data, this will query the AssetManager for you and takes current platform into account. Returns true and sets out parameter if found */
+	virtual bool GetStringValueForCustomColumn(const FAssetData& AssetData, FName ColumnName, FString& OutValue) = 0;
+
+	/** Gets the value of a "virtual" column for an asset data, this will query the AssetManager for you and takes current platform into account. Returns true and sets out parameter if found */
+	virtual bool GetIntegerValueForCustomColumn(const FAssetData& AssetData, FName ColumnName, int64& OutValue) = 0;
+
+	/** Returns the set of asset packages managed the given asset data, returns true if any found. This handles the fake chunk/primary asset types above */
+	virtual bool GetManagedPackageListForAssetData(const FAssetData& AssetData, TSet<FName>& ManagedPackageSet) = 0;
+
+	/** Gets the list of registry sources that are available */
+	virtual void GetAvailableRegistrySources(TArray<const FAssetManagerEditorRegistrySource*>& AvailableSources) = 0;
+
+	/** Returns the currently selected registry source. If bNeedManagementData is set it it will cache the chunk/management information */
+	virtual const FAssetManagerEditorRegistrySource* GetCurrentRegistrySource(bool bNeedManagementData = false) = 0;
+
+	/** Sets the current registry source, this loads the asset registry state if needed and may spawn a file load dialog for custom */
+	virtual void SetCurrentRegistrySource(const FString& SourceName) = 0;
+
+	/** Refreshes the management dictionary and all sources */
+	virtual void RefreshRegistryData() = 0;
+
+	/** Returns true if this package exists in the current registry source, optionally setting a redirected package name */
+	virtual bool IsPackageInCurrentRegistrySource(FName PackageName) = 0;
+
+	/** Filters list of identifiers and removes ones that do not exist in this registry source. Handles replacing redirectors as well */
+	virtual bool FilterAssetIdentifiersForCurrentRegistrySource(TArray<FAssetIdentifier>& AssetIdentifiers, EAssetRegistryDependencyType::Type DependencyType = EAssetRegistryDependencyType::None, bool bForwardDependency = true) = 0;
+
+	/** Creates a collection from a list of packages, will overwrite/modify an existing collection of the same name. Will display feedback to the user if bShowFeedback is true */
+	virtual bool WriteCollection(FName CollectionName, ECollectionShareType::Type ShareType, const TArray<FName>& PackageNames, bool bShowFeedback = true) = 0;
+
+	/** Opens asset management UI, with selected assets. Pass as value so it can be used in delegates */
+	virtual void OpenAssetAuditUI(TArray<FAssetData> SelectedAssets) = 0;
+	virtual void OpenAssetAuditUI(TArray<FName> SelectedPackages) = 0;
 
-	/** Gets the list of target platforms that are available */
-	virtual void GetAvailableTargetPlatforms(TArray<ITargetPlatform*>& AvailablePlatforms) = 0;
+	/** Spawns reference viewer, showing selected packages or identifiers */
+	virtual void OpenReferenceViewerUI(TArray<FAssetIdentifier> SelectedIdentifiers) = 0;
+	virtual void OpenReferenceViewerUI(TArray<FName> SelectedPackages) = 0;
 
-	/** Returns the asset registry state used by a specific target platform. This will load it on demand if needed */
-	virtual FAssetRegistryState* GetAssetRegistryStateForTargetPlatform(ITargetPlatform* TargetPlatform) = 0;
+	/** Spawns size map with selected packages */
+	virtual void OpenSizeMapUI(TArray<FAssetIdentifier> SelectedIdentifiers) = 0;
+	virtual void OpenSizeMapUI(TArray<FName> SelectedPackages) = 0;
 };
diff --git a/Engine/Plugins/Editor/GameplayTagsEditor/GameplayTagsEditor.uplugin b/Engine/Plugins/Editor/GameplayTagsEditor/GameplayTagsEditor.uplugin
--- a/Engine/Plugins/Editor/GameplayTagsEditor/GameplayTagsEditor.uplugin
+++ b/Engine/Plugins/Editor/GameplayTagsEditor/GameplayTagsEditor.uplugin
@@ -21,5 +21,12 @@
 			"Type" : "Developer",
 			"LoadingPhase" : "PreDefault"
 		}
+	],
+	"Plugins": 
+	[
+		{
+			"Name" : "AssetManagerEditor",
+			"Enabled" : true
+		}
 	]
 }
\ No newline at end of file
diff --git a/Engine/Plugins/Editor/GameplayTagsEditor/Source/GameplayTagsEditor/GameplayTagsEditor.Build.cs b/Engine/Plugins/Editor/GameplayTagsEditor/Source/GameplayTagsEditor/GameplayTagsEditor.Build.cs
--- a/Engine/Plugins/Editor/GameplayTagsEditor/Source/GameplayTagsEditor/GameplayTagsEditor.Build.cs
+++ b/Engine/Plugins/Editor/GameplayTagsEditor/Source/GameplayTagsEditor/GameplayTagsEditor.Build.cs
@@ -42,6 +42,7 @@ namespace UnrealBuildTool.Rules
 					"MainFrame",
 					"UnrealEd",
 					"SourceControl",
+					"AssetManagerEditor"
 				}
 			);
 
diff --git a/Engine/Plugins/Editor/GameplayTagsEditor/Source/GameplayTagsEditor/Private/SGameplayTagWidget.cpp b/Engine/Plugins/Editor/GameplayTagsEditor/Source/GameplayTagsEditor/Private/SGameplayTagWidget.cpp
--- a/Engine/Plugins/Editor/GameplayTagsEditor/Source/GameplayTagsEditor/Private/SGameplayTagWidget.cpp
+++ b/Engine/Plugins/Editor/GameplayTagsEditor/Source/GameplayTagsEditor/Private/SGameplayTagWidget.cpp
@@ -27,7 +27,7 @@
 #include "SAddNewGameplayTagWidget.h"
 #include "SRenameGameplayTagDialog.h"
 #include "AssetData.h"
-#include "Editor/ReferenceViewer/Public/ReferenceViewer.h"
+#include "AssetManagerEditorModule.h"
 #include "Framework/Commands/UIAction.h"
 #include "Framework/MultiBox/MultiBoxBuilder.h"
 
@@ -681,7 +681,7 @@ TSharedRef<SWidget> SGameplayTagWidget::MakeTagActionsMenu(TSharedPtr<FGameplayT
 	}
 
 	// Search for References
-	if (IReferenceViewerModule::IsAvailable())
+	if (IAssetManagerEditorModule::IsAvailable())
 	{
 		FExecuteAction SearchForReferencesAction = FExecuteAction::CreateSP(this, &SGameplayTagWidget::OnSearchForReferences, InTagNode);
 
@@ -717,14 +717,14 @@ void SGameplayTagWidget::OnDeleteTag(TSharedPtr<FGameplayTagNode> InTagNode)
 
 void SGameplayTagWidget::OnSearchForReferences(TSharedPtr<FGameplayTagNode> InTagNode)
 {
-	if (InTagNode.IsValid() && IReferenceViewerModule::IsAvailable())
+	if (InTagNode.IsValid() && IAssetManagerEditorModule::IsAvailable())
 	{
-		IReferenceViewerModule& ReferenceViewer = IReferenceViewerModule::Get();
+		IAssetManagerEditorModule& ManagerEditorModule = IAssetManagerEditorModule::Get();
 
 		TArray<FAssetIdentifier> AssetIdentifiers;
 		AssetIdentifiers.Add(FAssetIdentifier(FGameplayTag::StaticStruct(), InTagNode->GetCompleteTagName()));
 
-		ReferenceViewer.InvokeReferenceViewerTab(AssetIdentifiers);
+		ManagerEditorModule.OpenReferenceViewerUI(AssetIdentifiers);
 	}
 }
 
diff --git a/Engine/Plugins/Online/OnlineSubsystemGoogle/Source/OnlineSubsystemGoogle_UPL.xml b/Engine/Plugins/Online/OnlineSubsystemGoogle/Source/OnlineSubsystemGoogle_UPL.xml
--- a/Engine/Plugins/Online/OnlineSubsystemGoogle/Source/OnlineSubsystemGoogle_UPL.xml
+++ b/Engine/Plugins/Online/OnlineSubsystemGoogle/Source/OnlineSubsystemGoogle_UPL.xml
@@ -31,7 +31,7 @@
 	<AARImports>
 		<if condition="bEnabled">
 			<true>
-				<insertValue value="com.google.android.gms,play-services-auth,9.8.0" />
+				<insertValue value="com.google.android.gms,play-services-auth,11.0.4" />
 				<insertNewline/>
 			</true>
 		</if>
diff --git a/Engine/Plugins/Runtime/Analytics/Adjust/Source/IOSAdjust/IOSAdjust.Build.cs b/Engine/Plugins/Runtime/Analytics/Adjust/Source/IOSAdjust/IOSAdjust.Build.cs
--- a/Engine/Plugins/Runtime/Analytics/Adjust/Source/IOSAdjust/IOSAdjust.Build.cs
+++ b/Engine/Plugins/Runtime/Analytics/Adjust/Source/IOSAdjust/IOSAdjust.Build.cs
@@ -32,11 +32,11 @@ namespace UnrealBuildTool.Rules
 					"AdjustSdk",													// Framework name
 					"../../ThirdPartyFrameworks/AdjustSdk.embeddedframework.zip"	// Zip name
 				)
-			); 
+			);
 
 			PublicFrameworks.AddRange(
 				new string[] {
-//					"AdSuppport",	// optional
+//					"AdSupport",	// optional
 //					"iAd",			// optional
 				}
 			);
diff --git a/Engine/Plugins/Runtime/WebBrowserWidget/Source/WebBrowserWidget/Private/WebBrowser.cpp b/Engine/Plugins/Runtime/WebBrowserWidget/Source/WebBrowserWidget/Private/WebBrowser.cpp
--- a/Engine/Plugins/Runtime/WebBrowserWidget/Source/WebBrowserWidget/Private/WebBrowser.cpp
+++ b/Engine/Plugins/Runtime/WebBrowserWidget/Source/WebBrowserWidget/Private/WebBrowser.cpp
@@ -4,6 +4,7 @@
 #include "SWebBrowser.h"
 #include "Widgets/Layout/SBox.h"
 #include "Widgets/Text/STextBlock.h"
+#include "TaskGraphInterfaces.h"
 
 #define LOCTEXT_NAMESPACE "WebBrowser"
 
@@ -85,7 +86,8 @@ TSharedRef<SWidget> UWebBrowser::RebuildWidget()
 			.InitialURL(InitialURL)
 			.ShowControls(false)
 			.SupportsTransparency(bSupportsTransparency)
-			.OnUrlChanged(BIND_UOBJECT_DELEGATE(FOnTextChanged, HandleOnUrlChanged));
+			.OnUrlChanged(BIND_UOBJECT_DELEGATE(FOnTextChanged, HandleOnUrlChanged))
+			.OnBeforePopup(BIND_UOBJECT_DELEGATE(FOnBeforePopupDelegate, HandleOnBeforePopup));
 
 		return WebBrowserWidget.ToSharedRef();
 	}
@@ -106,6 +108,33 @@ void UWebBrowser::HandleOnUrlChanged(const FText& InText)
 	OnUrlChanged.Broadcast(InText);
 }
 
+bool UWebBrowser::HandleOnBeforePopup(FString URL, FString Frame)
+{
+	if (OnBeforePopup.IsBound())
+	{
+		if (IsInGameThread())
+		{
+			OnBeforePopup.Broadcast(URL, Frame);
+		}
+		else
+		{
+			// Retry on the GameThread.
+			TWeakObjectPtr<UWebBrowser> WeakThis = this;
+			FFunctionGraphTask::CreateAndDispatchWhenReady([WeakThis, URL, Frame]()
+			{
+				if (WeakThis.IsValid())
+				{
+					WeakThis->HandleOnBeforePopup(URL, Frame);
+				}
+			}, TStatId(), nullptr, ENamedThreads::GameThread);
+		}
+
+		return true;
+	}
+
+	return false;
+}
+
 #if WITH_EDITOR
 
 const FText UWebBrowser::GetPaletteCategory()
diff --git a/Engine/Plugins/Runtime/WebBrowserWidget/Source/WebBrowserWidget/Public/WebBrowser.h b/Engine/Plugins/Runtime/WebBrowserWidget/Source/WebBrowserWidget/Public/WebBrowser.h
--- a/Engine/Plugins/Runtime/WebBrowserWidget/Source/WebBrowserWidget/Public/WebBrowser.h
+++ b/Engine/Plugins/Runtime/WebBrowserWidget/Source/WebBrowserWidget/Public/WebBrowser.h
@@ -16,6 +16,7 @@ class WEBBROWSERWIDGET_API UWebBrowser : public UWidget
 
 public:
 	DECLARE_DYNAMIC_MULTICAST_DELEGATE_OneParam(FOnUrlChanged, const FText&, Text);
+	DECLARE_DYNAMIC_MULTICAST_DELEGATE_TwoParams(FOnBeforePopup, FString, URL, FString, Frame);
 
 	/**
 	 * Load the specified URL
@@ -60,6 +61,10 @@ public:
 	UPROPERTY(BlueprintAssignable, Category = "Web Browser|Event")
 	FOnUrlChanged OnUrlChanged;
 
+	/** Called when a popup is about to spawn. */
+	UPROPERTY(BlueprintAssignable, Category = "Web Browser|Event")
+	FOnBeforePopup OnBeforePopup;
+
 public:
 
 	//~ Begin UWidget interface
@@ -90,4 +95,5 @@ protected:
 	// End of UWidget interface
 
 	void HandleOnUrlChanged(const FText& Text);
+	bool HandleOnBeforePopup(FString URL, FString Frame);
 };
diff --git a/Engine/Shaders/Private/LocalVertexFactory.ush b/Engine/Shaders/Private/LocalVertexFactory.ush
--- a/Engine/Shaders/Private/LocalVertexFactory.ush
+++ b/Engine/Shaders/Private/LocalVertexFactory.ush
@@ -109,6 +109,7 @@ float3 SplineMeshY;
 
 	uint3 VertexFetch_Parameters;
 	uint InstanceOffset;
+	uint VertexOffset;
 
 	#define VF_ColorIndexMask_Index 0
 	#define VF_NumTexcoords_Index 1
@@ -407,7 +408,7 @@ FMaterialVertexParameters GetMaterialVertexParameters(FVertexFactoryInput Input,
 		UNROLL
 		for (uint CoordinateIndex = 0; CoordinateIndex < NUM_MATERIAL_TEXCOORDS_VERTEX; CoordinateIndex++)
 		{
-			Result.TexCoords[CoordinateIndex] = VertexFetch_TexCoordBuffer[VertexFetch_Parameters[VF_NumTexcoords_Index] * Input.VertexId + CoordinateIndex];
+			Result.TexCoords[CoordinateIndex] = VertexFetch_TexCoordBuffer[VertexFetch_Parameters[VF_NumTexcoords_Index] * (VertexOffset + Input.VertexId) + CoordinateIndex];
 		}
 #elif NUM_MATERIAL_TEXCOORDS_VERTEX
 		#if GPUSKIN_PASS_THROUGH
@@ -563,8 +564,8 @@ half3x3 CalcTangentToLocal(FVertexFactoryInput Input, out float TangentSign)
 	half3x3 Result;
 	
 #if MANUAL_VERTEX_FETCH
-	half3 TangentInputX = VertexFetch_PackedTangentsBuffer[2 * Input.VertexId + 0].xyz;
-	half4 TangentInputZ = VertexFetch_PackedTangentsBuffer[2 * Input.VertexId + 1].xyzw;
+	half3 TangentInputX = VertexFetch_PackedTangentsBuffer[2 * (VertexOffset + Input.VertexId) + 0].xyz;
+	half4 TangentInputZ = VertexFetch_PackedTangentsBuffer[2 * (VertexOffset + Input.VertexId) + 1].xyzw;
 #else
 	half3 TangentInputX = Input.TangentX;
 	half4 TangentInputZ = Input.TangentZ;
@@ -615,7 +616,7 @@ FVertexFactoryIntermediates GetVertexFactoryIntermediates(FVertexFactoryInput In
 	FVertexFactoryIntermediates Intermediates;
 
 #if MANUAL_VERTEX_FETCH
-	Intermediates.Color = VertexFetch_ColorComponentsBuffer[Input.VertexId & VertexFetch_Parameters[VF_ColorIndexMask_Index]] FMANUALFETCH_COLOR_COMPONENT_SWIZZLE; // Swizzle vertex color.
+	Intermediates.Color = VertexFetch_ColorComponentsBuffer[(VertexOffset + Input.VertexId) & VertexFetch_Parameters[VF_ColorIndexMask_Index]] FMANUALFETCH_COLOR_COMPONENT_SWIZZLE; // Swizzle vertex color.
 #else
 	Intermediates.Color = Input.Color FCOLOR_COMPONENT_SWIZZLE; // Swizzle vertex color.
 #endif
@@ -717,7 +718,7 @@ FVertexFactoryInterpolantsVSToPS VertexFactoryGetInterpolantsVSToPS(FVertexFacto
 
 #elif NUM_MATERIAL_TEXCOORDS_VERTEX == 0 && USE_PARTICLE_SUBUVS
 	#if MANUAL_VERTEX_FETCH
-		SetUV(Interpolants, 0, VertexFetch_TexCoordBuffer[VertexFetch_Parameters[VF_NumTexcoords_Index] * Input.VertexId]);
+		SetUV(Interpolants, 0, VertexFetch_TexCoordBuffer[VertexFetch_Parameters[VF_NumTexcoords_Index] * (VertexOffset + Input.VertexId)]);
 	#else
 		SetUV(Interpolants, 0, Input.TexCoords[0]);
 	#endif
@@ -727,7 +728,7 @@ FVertexFactoryInterpolantsVSToPS VertexFactoryGetInterpolantsVSToPS(FVertexFacto
 	float2 LightMapCoordinate = 0;
 	float2 ShadowMapCoordinate = 0;
 	#if MANUAL_VERTEX_FETCH
-		float2 LightMapCoordinateInput = VertexFetch_TexCoordBuffer[VertexFetch_Parameters[VF_NumTexcoords_Index] * Input.VertexId + VertexFetch_Parameters[FV_LightMapIndex_Index]];
+		float2 LightMapCoordinateInput = VertexFetch_TexCoordBuffer[VertexFetch_Parameters[VF_NumTexcoords_Index] * (VertexOffset + Input.VertexId) + VertexFetch_Parameters[FV_LightMapIndex_Index]];
 	#else
 		float2 LightMapCoordinateInput = Input.LightMapCoordinate;
 	#endif
diff --git a/Engine/Source/Developer/TreeMap/Public/ITreeMap.h b/Engine/Source/Developer/TreeMap/Public/ITreeMap.h
--- a/Engine/Source/Developer/TreeMap/Public/ITreeMap.h
+++ b/Engine/Source/Developer/TreeMap/Public/ITreeMap.h
@@ -52,8 +52,10 @@ class FTreeMapNodeData : public TSharedFromThis<FTreeMapNodeData>
 {
 
 public:
+	/** Logical name, if set this is used for comparison */
+	FString LogicalName;
 
-	/** Node name.  When drawing the tree map, this shows up at the top of the inside of the node's rectangle */
+	/** Node name display.  When drawing the tree map, this shows up at the top of the inside of the node's rectangle */
 	FString Name;
 
 	/** Node name (line 2).  Left nodes only.  This shows up on the second line of the node, right underneath its name.  Might be drawn in a smaller font by default. */
@@ -62,7 +64,6 @@ public:
 	/** Center info text.  Leaf nodes only.  This draws right in the center of the node, with a larger font by default */
 	FString CenterText;
 
-
 	/** 
 	 * Size of this node.  This works differently depending on whether the node is a leaf node (no children) or a container node (has children):
 	 *  - Leaf nodes must ALWAYS have a non-zero size.  
@@ -90,21 +91,30 @@ public:
 	/** List of child nodes */
 	TArray< FTreeMapNodeDataPtr > Children;
 
-
 public:
 
 	/** Default constructor for FTreeMapNodeData */
 	FTreeMapNodeData()
-		: Name(),
-		  Name2(),
-		  CenterText(),
-		  Size( 0.0 ),
+		: Size( 0.0 ),
 		  BackgroundBrush( nullptr ),
 		  Color( FLinearColor::White ),
-		  HashTags(),
-		  Parent( NULL ),
-		  Children()
+		  Parent( NULL )
+	{
+	}
+
+	bool operator==(const FTreeMapNodeData& Other) const
+	{
+		if (!LogicalName.IsEmpty() || !Other.LogicalName.IsEmpty())
+		{
+			return LogicalName == Other.LogicalName;
+		}
+
+		return Name == Other.Name && Name2 == Other.Name2 && CenterText == Other.CenterText;
+	}
+
+	bool operator!=(const FTreeMapNodeData& Other) const
 	{
+		return !(*this == Other);
 	}
 
 	/** @return Returns true if this is a leaf node */
@@ -116,6 +126,7 @@ public:
 	/** @return Copies this node into another node (not children, though.  The copied node will have no children.) */
 	void CopyNodeInto( FTreeMapNodeData& NodeCopy ) const
 	{
+		NodeCopy.LogicalName = this->LogicalName;
 		NodeCopy.Name = this->Name;
 		NodeCopy.Name2 = this->Name2;
 		NodeCopy.CenterText = this->CenterText;
@@ -204,6 +215,9 @@ struct FTreeMapOptions
 	/** Minimize size of a tree node that will be allowed to have a title and padding.  Nodes smaller than this will be a simple rectangle.  Doesn't apply to top level nodes. */
 	float MinimumInteractiveNodeSize;
 
+	/** Minimize size of a tree node that will be drawn at all, below this size it is invisible */
+	float MinimumVisibleNodeSize;
+
 
 	/** Default constructor for FTreeMapOptions that initializes good defaults */
 	FTreeMapOptions()
@@ -217,7 +231,8 @@ struct FTreeMapOptions
 		  TopLevelContainerOuterPadding( 0.0f ),
 		  NestedContainerOuterPadding( 0.0f ),
 		  ContainerInnerPadding( 0.0f ),
-		  MinimumInteractiveNodeSize( 0.0f )
+		  MinimumInteractiveNodeSize( 0.0f ),
+		  MinimumVisibleNodeSize( 0.0f )
 	{
 		NameFont.Size = 12;
 		Name2Font.Size = 8;
diff --git a/Engine/Source/Developer/TreeMap/Public/STreeMap.h b/Engine/Source/Developer/TreeMap/Public/STreeMap.h
--- a/Engine/Source/Developer/TreeMap/Public/STreeMap.h
+++ b/Engine/Source/Developer/TreeMap/Public/STreeMap.h
@@ -25,11 +25,8 @@ class TREEMAP_API STreeMap : public SLeafWidget
 
 public:
 
-	/** Optional delegate that fires when the node is double-clicked in the tree.  If you don't override this, the tree will use its
-	    default handling of double-click, which will re-root the tree on the node under the cursor */
-	DECLARE_DELEGATE_OneParam( FOnTreeMapNodeDoubleClicked, FTreeMapNodeData& );
-
-
+	/** Delegate used when clicking or interacting with a specific node */
+	DECLARE_DELEGATE_TwoParams( FOnTreeMapNodeInteracted, FTreeMapNodeData&, const FPointerEvent& );
 
 	SLATE_BEGIN_ARGS( STreeMap )
 		: _AllowEditing( false )
@@ -40,7 +37,8 @@ public:
 		, _Name2Font()
 		, _CenterTextFont()
 		, _BorderPadding( FTreeMapStyle::Get().GetVector( "TreeMap.BorderPadding" ) )
-		, _MinimumVisualTreeNodeSize( 64 * 64 )
+		, _MinimumInteractiveTreeNodeSize( 64 * 64 )
+		, _MinimumVisibleTreeNodeSize( 6 * 6 )
 		, _NavigationTransitionTime( 0.25f )
 		, _TopLevelContainerOuterPadding( 4.0f )
 		, _NestedContainerOuterPadding( 0.0f )
@@ -74,8 +72,11 @@ public:
 		/** Border Padding around fill bar */
 		SLATE_ATTRIBUTE( FVector2D, BorderPadding )
 
+		/** Minimum size of node that can be interacted with, if small you need to drill down into parent first */
+		SLATE_ARGUMENT( int32, MinimumInteractiveTreeNodeSize );
+
 		/** Minimum size in pixels of a tree node that we should bother including in the UI.  Below this size, you'll need to drill down to see the node. */
-		SLATE_ARGUMENT( int32, MinimumVisualTreeNodeSize );
+		SLATE_ARGUMENT( int32, MinimumVisibleTreeNodeSize );
 
 		/** How many seconds to animate the visual transition when the user navigates to a new tree node, or after a modification of the tree takes place */
 		SLATE_ARGUMENT( float, NavigationTransitionTime );
@@ -94,7 +95,10 @@ public:
 
 		/** Optional delegate that fires when the node is double-clicked in the tree.  If you don't override this, the tree will use its
 			default handling of double-click, which will re-root the tree on the node under the cursor */
-		SLATE_EVENT( FOnTreeMapNodeDoubleClicked, OnTreeMapNodeDoubleClicked )
+		SLATE_EVENT( FOnTreeMapNodeInteracted, OnTreeMapNodeDoubleClicked )
+
+		/** Optional delegate that fires when the node is right-clicked in the tree.  If you don't override this nothing will happen unless a customization is specified */
+		SLATE_EVENT( FOnTreeMapNodeInteracted, OnTreeMapNodeRightClicked )
 
 	SLATE_END_ARGS()
 
@@ -129,6 +133,15 @@ public:
 	 */
 	void SetTreeRoot( const FTreeMapNodeDataRef& NewRoot, const bool bShouldPlayTransition );
 
+	/** Returns true if it is possible to zoom out */
+	bool CanZoomOut() const;
+
+	/** Zooms out, setting root to it's parent if possible, returns false if not */
+	bool ZoomOut();
+
+	/** Returns the current tree root, which is the largest node visible */
+	FTreeMapNodeDataPtr GetTreeRoot() const;
+
 	/**
 	 * Refreshes the tree map from its source data.  The current active root will remain as the base of the tree.  Call this
 	 * after you've changed data in the tree.
@@ -254,15 +267,20 @@ private:
 
 	/** Optional delegate that fires when the node is double-clicked in the tree.  If you don't override this, the tree will use its
 		default handling of double-click, which will re-root the tree on the node under the cursor */
-	FOnTreeMapNodeDoubleClicked OnTreeMapNodeDoubleClicked;
+	FOnTreeMapNodeInteracted OnTreeMapNodeDoubleClicked;
 
+	/** Right click event */
+	FOnTreeMapNodeInteracted OnTreeMapNodeRightClicked;
 
 	/**
 	 * TreeMap visuals
 	 */
 
+	/** Minimum size of node that can be interacted with, if small you need to drill down into parent first */
+	int32 MinimumInteractiveTreeNodeSize;
+
 	/** Minimum size in pixels of a tree node that we should bother including in the UI.  Below this size, you'll need to drill down to see the node. */
-	int32 MinimumVisualTreeNodeSize;
+	int32 MinimumVisibleTreeNodeSize;
 
 	/** Size of our geometry the last time we rebuilt the tree map */
 	FVector2D TreeMapSize;
diff --git a/Engine/Source/Developer/TreeMap/STreeMap.cpp b/Engine/Source/Developer/TreeMap/STreeMap.cpp
--- a/Engine/Source/Developer/TreeMap/STreeMap.cpp
+++ b/Engine/Source/Developer/TreeMap/STreeMap.cpp
@@ -61,7 +61,8 @@ void STreeMap::Construct( const FArguments& InArgs, const TSharedRef<FTreeMapNod
 		}
 	}
 
-	MinimumVisualTreeNodeSize = InArgs._MinimumVisualTreeNodeSize;
+	MinimumInteractiveTreeNodeSize = InArgs._MinimumInteractiveTreeNodeSize;
+	MinimumVisibleTreeNodeSize = InArgs._MinimumVisibleTreeNodeSize;
 	TopLevelContainerOuterPadding = InArgs._TopLevelContainerOuterPadding;
 	NestedContainerOuterPadding = InArgs._NestedContainerOuterPadding;
 	ContainerInnerPadding = InArgs._ContainerInnerPadding;
@@ -69,8 +70,8 @@ void STreeMap::Construct( const FArguments& InArgs, const TSharedRef<FTreeMapNod
 
 	TreeMapSize = FVector2D( 0, 0 );
 	NavigateAnimationCurve.AddCurve( 0.0f, InArgs._NavigationTransitionTime, ECurveEaseFunction::CubicOut );
-	MouseOverVisual = NULL;
-	DraggingVisual = NULL;
+	MouseOverVisual = nullptr;
+	DraggingVisual = nullptr;
 	DragVisualDistance = 0.0f;
 
 	bIsNamingNewNode = false;
@@ -78,6 +79,7 @@ void STreeMap::Construct( const FArguments& InArgs, const TSharedRef<FTreeMapNod
 	HighlightPulseStartTime = -99999.0;
 
 	OnTreeMapNodeDoubleClicked = InArgs._OnTreeMapNodeDoubleClicked;
+	OnTreeMapNodeRightClicked = InArgs._OnTreeMapNodeRightClicked;
 
 	if( Customization.IsValid() )
 	{
@@ -111,14 +113,15 @@ void STreeMap::Tick( const FGeometry& AllottedGeometry, const double InCurrentTi
 		TreeMapOptions.Name2Font = Name2Font.Get();
 		TreeMapOptions.CenterTextFont = CenterTextFont.Get();
 		TreeMapOptions.FontSizeChangeBasedOnDepth = 2;		// @todo treemap custom: Expose as a customization option to STreeMap
-		TreeMapOptions.MinimumInteractiveNodeSize = MinimumVisualTreeNodeSize;
+		TreeMapOptions.MinimumInteractiveNodeSize = MinimumInteractiveTreeNodeSize;
+		TreeMapOptions.MinimumVisibleNodeSize = MinimumVisibleTreeNodeSize;
 		TreeMap = ITreeMap::CreateTreeMap( TreeMapOptions, ActiveRootTreeMapNode.ToSharedRef() );
 
 		TreeMapSize = AllottedGeometry.Size;
 
 		CachedNodeVisuals = TreeMap->GetVisuals();
-		MouseOverVisual = NULL;
-		DraggingVisual = NULL;
+		MouseOverVisual = nullptr;
+		DraggingVisual = nullptr;
 
 		// Map the new visuals back to the old ones!
 		NodeVisualIndicesToLastIndices.Reset();
@@ -160,7 +163,7 @@ void STreeMap::MakeBlendedNodeVisual( const int32 VisualIndex, const float Navig
 	{
 		// Did the visual exist before we navigated?
 		const int32* LastVisualIndexPtr = NodeVisualIndicesToLastIndices.Find( VisualIndex );
-		if( LastVisualIndexPtr != NULL )
+		if( LastVisualIndexPtr != nullptr )
 		{
 			// It did exist!
 			const auto LastVisualIndex = *LastVisualIndexPtr;
@@ -253,7 +256,7 @@ int32 STreeMap::OnPaint( const FPaintArgs& Args, const FGeometry& AllottedGeomet
 				// Don't draw if completely faded out
 				if( BlendedVisual.Color.A > KINDA_SMALL_NUMBER )
 				{
-					const bool bIsMouseOverNode = MouseOverVisual != NULL && MouseOverVisual->NodeData == BlendedVisual.NodeData;
+					const bool bIsMouseOverNode = MouseOverVisual != nullptr && MouseOverVisual->NodeData == BlendedVisual.NodeData;
 
 					// Draw the visual's background box
 					const FVector2D VisualPosition = BlendedVisual.Position;
@@ -370,7 +373,7 @@ int32 STreeMap::OnPaint( const FPaintArgs& Args, const FGeometry& AllottedGeomet
 						auto VisualTextColor = TextColorScale;
 						VisualTextColor.A *= BlendedVisual.Color.A;
 
-						const bool bIsMouseOverNode = MouseOverVisual != NULL && MouseOverVisual->NodeData == BlendedVisual.NodeData;
+						const bool bIsMouseOverNode = MouseOverVisual != nullptr && MouseOverVisual->NodeData == BlendedVisual.NodeData;
 
 						// If the visual is crushed down pretty small in screen space, we don't want to bother even try drawing text
 						const FVector2D ScreenSpaceVisualSize( AllottedGeometry.Scale * BlendedVisual.Size );
@@ -461,7 +464,7 @@ int32 STreeMap::OnPaint( const FPaintArgs& Args, const FGeometry& AllottedGeomet
 	}
 
 
-	if( DraggingVisual != NULL && DragVisualDistance >= STreeMapDefs::MinCursorDistanceForDraggingVisual )
+	if( DraggingVisual != nullptr && DragVisualDistance >= STreeMapDefs::MinCursorDistanceForDraggingVisual )
 	{
 		const FVector2D DragStartCursorPos = RelativeDragStartMouseCursorPos;
 		const FVector2D NewCursorPos = RelativeMouseCursorPos;
@@ -516,14 +519,14 @@ FReply STreeMap::OnMouseMove( const FGeometry& InMyGeometry, const FPointerEvent
 	RelativeMouseCursorPos = InMyGeometry.AbsoluteToLocal( InMouseEvent.GetScreenSpacePosition() );
 
 	// Clear current hover
-	MouseOverVisual = NULL;
+	MouseOverVisual = nullptr;
 
 	// Don't hover while transitioning.  It looks bad!
 	if( !IsNavigationTransitionActive() )
 	{
 		// Figure out which visual the cursor is over
 		FTreeMapNodeVisualInfo* NodeVisualUnderCursor = FindNodeVisualUnderCursor( InMyGeometry, InMouseEvent.GetScreenSpacePosition() );
-		if( NodeVisualUnderCursor != NULL )
+		if( NodeVisualUnderCursor != nullptr )
 		{
 			// Mouse is over a visual
 			MouseOverVisual = NodeVisualUnderCursor;
@@ -532,7 +535,7 @@ FReply STreeMap::OnMouseMove( const FGeometry& InMyGeometry, const FPointerEvent
 		Reply = FReply::Handled();
 	}
 
-	if( DraggingVisual != NULL )
+	if( DraggingVisual != nullptr )
 	{
 		DragVisualDistance += InMouseEvent.GetCursorDelta().Size();
 	}
@@ -542,7 +545,7 @@ FReply STreeMap::OnMouseMove( const FGeometry& InMyGeometry, const FPointerEvent
 
 void STreeMap::OnMouseLeave( const FPointerEvent& InMouseEvent )
 {
-	MouseOverVisual = NULL;
+	MouseOverVisual = nullptr;
 }
 
 FReply STreeMap::OnMouseButtonDown( const FGeometry& InMyGeometry, const FPointerEvent& InMouseEvent )
@@ -557,7 +560,7 @@ FReply STreeMap::OnMouseButtonDown( const FGeometry& InMyGeometry, const FPointe
 			{
 				// Figure out what was clicked on
 				FTreeMapNodeVisualInfo* NodeVisualUnderCursor = FindNodeVisualUnderCursor( InMyGeometry, InMouseEvent.GetScreenSpacePosition() );
-				if( NodeVisualUnderCursor != NULL )
+				if( NodeVisualUnderCursor != nullptr )
 				{
 					// Mouse was pressed on a node
 					DraggingVisual = NodeVisualUnderCursor;
@@ -576,44 +579,36 @@ FReply STreeMap::OnMouseButtonDown( const FGeometry& InMyGeometry, const FPointe
 FReply STreeMap::OnMouseButtonUp( const FGeometry& InMyGeometry, const FPointerEvent& InMouseEvent )
 {
 	FReply Reply = FReply::Unhandled();
+	FTreeMapNodeVisualInfo* NodeVisualUnderCursor = FindNodeVisualUnderCursor(InMyGeometry, InMouseEvent.GetScreenSpacePosition());
 	if( InMouseEvent.GetEffectingButton() == EKeys::LeftMouseButton )
 	{
 		auto* DroppedVisual = DraggingVisual;
-		DraggingVisual = NULL;
+		DraggingVisual = nullptr;
 
-		if( DroppedVisual != NULL && DragVisualDistance >= STreeMapDefs::MinCursorDistanceForDraggingVisual )
+		if( DroppedVisual != nullptr && DragVisualDistance >= STreeMapDefs::MinCursorDistanceForDraggingVisual )
 		{
 			// Wait until we've finished animating a transition
-			if( !IsNavigationTransitionActive() )
+			if( !IsNavigationTransitionActive() && NodeVisualUnderCursor != nullptr )
 			{
-				// Find the node under the cursor
-				FTreeMapNodeVisualInfo* NodeVisualUnderCursor = FindNodeVisualUnderCursor( InMyGeometry, InMouseEvent.GetScreenSpacePosition() );
-				if( NodeVisualUnderCursor != NULL )
-				{
-					// The dropped node will become a child of the node we dropped onto
-					auto NewParentNode = NodeVisualUnderCursor->NodeData->AsShared();
-					auto DroppedNode = DroppedVisual->NodeData->AsShared();
+				// The dropped node will become a child of the node we dropped onto
+				auto NewParentNode = NodeVisualUnderCursor->NodeData->AsShared();
+				auto DroppedNode = DroppedVisual->NodeData->AsShared();
 
-					// Reparent it!
-					ReparentNode( DroppedNode, NewParentNode );
+				// Reparent it!
+				ReparentNode( DroppedNode, NewParentNode );
 
-					Reply = FReply::Handled();
-				}
+				Reply = FReply::Handled();
 			}
 		}
-		else
+		else if( NodeVisualUnderCursor != nullptr )
 		{
-			FTreeMapNodeVisualInfo* NodeVisualUnderCursor = FindNodeVisualUnderCursor( InMyGeometry, InMouseEvent.GetScreenSpacePosition() );
-			if( NodeVisualUnderCursor != NULL )
+			if( AllowEditing.Get() )
 			{
-				if( AllowEditing.Get() )
-				{
-					// Start renaming!
-					const bool bIsNewNode = false;
-					StartRenamingNode( InMyGeometry, NodeVisualUnderCursor->NodeData->AsShared(), NodeVisualUnderCursor->Position, bIsNewNode );
+				// Start renaming!
+				const bool bIsNewNode = false;
+				StartRenamingNode( InMyGeometry, NodeVisualUnderCursor->NodeData->AsShared(), NodeVisualUnderCursor->Position, bIsNewNode );
 
-					Reply = FReply::Handled();
-				}
+				Reply = FReply::Handled();
 			}
 		}
 
@@ -621,10 +616,37 @@ FReply STreeMap::OnMouseButtonUp( const FGeometry& InMyGeometry, const FPointerE
 	}
 	else if( InMouseEvent.GetEffectingButton() == EKeys::RightMouseButton )
 	{
-		// Show a pop-up menu!
-		ShowOptionsMenuAt( InMouseEvent );
+		if (OnTreeMapNodeRightClicked.IsBound())
+		{
+			if (NodeVisualUnderCursor != nullptr)
+			{
+				OnTreeMapNodeRightClicked.Execute(*NodeVisualUnderCursor->NodeData, InMouseEvent);
+			}
+		}
+		else
+		{
+			// Show a pop-up menu!
+			ShowOptionsMenuAt(InMouseEvent);
+		}
 	}
-
+	else if ( InMouseEvent.GetEffectingButton() == EKeys::ThumbMouseButton )
+	{
+		// Back button
+		if (ZoomOut() && !Reply.IsEventHandled())
+		{
+			Reply = FReply::Handled();
+		}
+	}
+	else if ( InMouseEvent.GetEffectingButton() == EKeys::ThumbMouseButton2 )
+	{
+		if (NodeVisualUnderCursor != nullptr)
+		{
+			// Do zoom behavior
+			const bool bShouldPlayTransition = true;
+			SetTreeRoot(NodeVisualUnderCursor->NodeData->AsShared(), bShouldPlayTransition);
+		}
+	}
+	
 	return Reply;
 }
 
@@ -639,12 +661,12 @@ FReply STreeMap::OnMouseButtonDoubleClick( const FGeometry& InMyGeometry, const
 		{
 			// Figure out what was clicked on
 			const FTreeMapNodeVisualInfo* NodeVisualUnderCursor = FindNodeVisualUnderCursor( InMyGeometry, InMouseEvent.GetScreenSpacePosition() );
-			if( NodeVisualUnderCursor != NULL )
+			if( NodeVisualUnderCursor != nullptr )
 			{
 				// Double-clicked on a tree map visual!  Check to see if we were asked to customize how double-click is handled.
 				if( OnTreeMapNodeDoubleClicked.IsBound() )
 				{
-					OnTreeMapNodeDoubleClicked.Execute( *NodeVisualUnderCursor->NodeData );
+					OnTreeMapNodeDoubleClicked.Execute( *NodeVisualUnderCursor->NodeData, InMouseEvent );
 				}
 				else
 				{
@@ -675,14 +697,14 @@ FReply STreeMap::OnMouseWheel( const FGeometry& MyGeometry, const FPointerEvent&
 		{
 			// Figure out what was clicked on
 			const FTreeMapNodeVisualInfo* NodeVisualUnderCursor = FindNodeVisualUnderCursor( MyGeometry, MouseEvent.GetScreenSpacePosition() );
-			if( NodeVisualUnderCursor != NULL )
+			if( NodeVisualUnderCursor != nullptr )
 			{
 				// From the node that was scrolled over, visits nodes upward until we find one whose parent is our current active root
 				FTreeMapNodeDataPtr NextNode = NodeVisualUnderCursor->NodeData->AsShared();
 				do
 				{
 					FTreeMapNodeDataPtr NodeParent;
-					if( NextNode->Parent != NULL )
+					if( NextNode->Parent != nullptr )
 					{
 						NodeParent = NextNode->Parent->AsShared();
 					}
@@ -711,7 +733,7 @@ FReply STreeMap::OnMouseWheel( const FGeometry& MyGeometry, const FPointerEvent&
 		else if( MouseEvent.GetWheelDelta() < 0 )
 		{
 			// Zoom out one level
-			if( ActiveRootTreeMapNode->Parent != NULL )
+			if( ActiveRootTreeMapNode->Parent != nullptr )
 			{
 				const bool bShouldPlayTransition = true;
 				SetTreeRoot( ActiveRootTreeMapNode->Parent->AsShared(), bShouldPlayTransition );
@@ -740,6 +762,36 @@ void STreeMap::SetTreeRoot( const FTreeMapNodeDataRef& NewRoot, const bool bShou
 }
 
 
+FTreeMapNodeDataPtr STreeMap::GetTreeRoot() const
+{
+	return ActiveRootTreeMapNode;
+}
+
+
+bool STreeMap::CanZoomOut() const
+{
+	if (ActiveRootTreeMapNode.IsValid() && ActiveRootTreeMapNode->Parent != nullptr)
+	{
+		return true;
+	}
+	return false;
+}
+
+
+bool STreeMap::ZoomOut()
+{
+	if (!CanZoomOut())
+	{
+		return false;
+	}
+
+	const bool bShouldPlayTransition = true;
+	SetTreeRoot(ActiveRootTreeMapNode->Parent->AsShared(), bShouldPlayTransition);
+
+	return true;
+}
+
+
 void STreeMap::RebuildTreeMap( const bool bShouldPlayTransition )
 {
 	// Stop renaming anything.  We don't want pop-up windows persisting during the transition
@@ -760,8 +812,20 @@ void STreeMap::RebuildTreeMap( const bool bShouldPlayTransition )
 	// Kill the tree map so that it will be regenerated
 	TreeMap.Reset();
 	CachedNodeVisuals.Reset();
-	DraggingVisual = NULL;
-	MouseOverVisual = NULL;
+	DraggingVisual = nullptr;
+	MouseOverVisual = nullptr;
+
+	// Update ActiveRootTreeMapNode
+	FTreeMapNodeDataPtr NewActiveRoot = FindNodeInCopiedTree(ActiveRootTreeMapNode.ToSharedRef(), TreeMapNodeData.ToSharedRef(), TreeMapNodeData.ToSharedRef());
+
+	if (NewActiveRoot.IsValid())
+	{
+		ActiveRootTreeMapNode = NewActiveRoot;
+	}
+	else
+	{
+		ActiveRootTreeMapNode = TreeMapNodeData;
+	}
 
 	// Customize the size and coloring of the source nodes before we rebuild it
 	ApplyVisualizationToNodes( ActiveRootTreeMapNode.ToSharedRef() );
@@ -791,7 +855,7 @@ FTreeMapNodeVisualInfo* STreeMap::FindNodeVisualUnderCursor( const FGeometry& My
 		}
 	}
 
-	return NULL;
+	return nullptr;
 }
 
 
@@ -844,7 +908,7 @@ FTreeMapNodeDataPtr STreeMap::FindNodeInCopiedTree( const FTreeMapNodeDataRef& N
 		}
 	}
 
-	return NULL;
+	return nullptr;
 }
 
 
@@ -936,7 +1000,7 @@ void STreeMap::ReparentNode( const FTreeMapNodeDataRef DroppedNode, const FTreeM
 					return true;
 				}
 
-				if( PossibleParent->Parent != NULL )
+				if( PossibleParent->Parent != nullptr )
 				{
 					return IsMyParent( Node, PossibleParent->Parent->AsShared() );
 				}
@@ -957,7 +1021,7 @@ void STreeMap::ReparentNode( const FTreeMapNodeDataRef DroppedNode, const FTreeM
 					// Take an Undo snapshot before we change anything
 					TakeUndoSnapshot();
 
-					if( DroppedNode->Parent != NULL )
+					if( DroppedNode->Parent != nullptr )
 					{
 						DroppedNode->Parent->Children.RemoveSingle( DroppedNode );
 					}
@@ -993,11 +1057,11 @@ FReply STreeMap::DeleteHoveredNode()
 	// Wait until we've finished animating a transition
 	if( !IsNavigationTransitionActive() )
 	{
-		if( MouseOverVisual != NULL )
+		if( MouseOverVisual != nullptr )
 		{
 			// Only non-root nodes can be deleted
 			auto NodeToDelete = MouseOverVisual->NodeData->AsShared();
-			if( NodeToDelete->Parent != NULL )
+			if( NodeToDelete->Parent != nullptr )
 			{
 				// Don't allow the current active tree root to be deleted.  The user should zoom out first!
 				if( NodeToDelete != ActiveRootTreeMapNode )
@@ -1007,7 +1071,7 @@ FReply STreeMap::DeleteHoveredNode()
 					// Delete the node
 					{
 						NodeToDelete->Parent->Children.RemoveSingle( NodeToDelete );
-						NodeToDelete->Parent = NULL;
+						NodeToDelete->Parent = nullptr;
 
 						// Note:  NodeToDelete will actually be deleted when it goes out of scope later in this function (shared pointer), but it is already removed from our tree
 					}
@@ -1032,7 +1096,7 @@ FReply STreeMap::InsertNewNodeAsChildOfHoveredNode( const FGeometry& MyGeometry
 	// Wait until we've finished animating a transition
 	if( !IsNavigationTransitionActive() )
 	{
-		if( MouseOverVisual != NULL )
+		if( MouseOverVisual != nullptr )
 		{
 			auto ParentNode = MouseOverVisual->NodeData->AsShared();
 
@@ -1042,7 +1106,6 @@ FReply STreeMap::InsertNewNodeAsChildOfHoveredNode( const FGeometry& MyGeometry
 			// in the delegate callback for the editable text change commit handler.  If the user opts to not type anything, the
 			// node will be deleted instead of being added to the tree.
 			FTreeMapNodeDataRef NewNode( new FTreeMapNodeData() );
-			NewNode->Name = TEXT( "" );
 			NewNode->Size = 1.0f;	// Leaf nodes always get a size of 1.0 for now
 			NewNode->Parent = MouseOverVisual->NodeData;
 
@@ -1155,7 +1218,7 @@ void STreeMap::StartRenamingNode( const FGeometry& MyGeometry, const FTreeMapNod
 
 	TSharedRef<SEditableTextBox> EditableText = 
 		SNew( SEditableTextBox )
-			.Text( FText::FromString( NodeData->Name ) )
+			.Text( FText::FromString( NodeData->Name) )
 			.SelectAllTextWhenFocused( true )
 			.RevertTextOnEscape( true )
 			.MinDesiredWidth( 100 )
@@ -1212,7 +1275,7 @@ void STreeMap::ApplyVisualizationToNodesRecursively( const FTreeMapNodeDataRef&
 		if( SizeNodesByAttribute.IsValid() )
 		{
 			FTreeMapAttributeDataPtr* AttributeDataPtr = Node->Attributes.Find( SizeNodesByAttribute->Name );
-			if( AttributeDataPtr != NULL )
+			if( AttributeDataPtr != nullptr )
 			{
 				const FTreeMapAttributeData& AttributeData = *AttributeDataPtr->Get();
 
@@ -1250,7 +1313,7 @@ void STreeMap::ApplyVisualizationToNodesRecursively( const FTreeMapNodeDataRef&
 		if( ColorNodesByAttribute.IsValid() )
 		{
 			FTreeMapAttributeDataPtr* AttributeDataPtr = Node->Attributes.Find( ColorNodesByAttribute->Name );
-			if( AttributeDataPtr != NULL )
+			if( AttributeDataPtr != nullptr )
 			{
 				const FTreeMapAttributeData& AttributeData = *AttributeDataPtr->Get();
 
@@ -1520,9 +1583,9 @@ void STreeMap::ShowOptionsMenuAtInternal(const FVector2D& ScreenSpacePosition, c
 	if( Customization.IsValid() )
 	{
 		const bool bShouldCloseMenuAfterSelection = true;
-		FMenuBuilder OptionsMenuBuilder( bShouldCloseMenuAfterSelection, NULL );
+		FMenuBuilder OptionsMenuBuilder( bShouldCloseMenuAfterSelection, nullptr );
 
-		if( AllowEditing.Get() && Customization.IsValid() && MouseOverVisual != NULL )
+		if( AllowEditing.Get() && Customization.IsValid() && MouseOverVisual != nullptr )
 		{
 			// Node editing options
 			OptionsMenuBuilder.BeginSection( NAME_None, LOCTEXT( "EditNodeAttributesSection", "Edit Node" ) );
diff --git a/Engine/Source/Developer/TreeMap/TreeMap.cpp b/Engine/Source/Developer/TreeMap/TreeMap.cpp
--- a/Engine/Source/Developer/TreeMap/TreeMap.cpp
+++ b/Engine/Source/Developer/TreeMap/TreeMap.cpp
@@ -70,6 +70,8 @@ public:
 	/** True if the node is 'interactive'.  That is, we have enough room for a title area and padding for the node to be clicked on */
 	bool bIsInteractive;
 
+	/** True if the node is visible at all, if false drawing will skip this entirely */
+	bool bIsVisible;
 
 public:
 
@@ -492,8 +494,10 @@ void FTreeMap::PadNodesRecursively( const FTreeMapOptions& Options, const FTreeM
 		auto ChildAreaRect = Node->PaddedRect;
 
 		// Unless this is a top level node, make sure the node is big enough to bother reporting to our caller.  They may not want to visualize tiny nodes!
+		Node->bIsVisible = ( ChildAreaRect.Size.X * ChildAreaRect.Size.Y >= Options.MinimumVisibleNodeSize );
 		Node->bIsInteractive = ( TreeDepth <= 1 || ChildAreaRect.Size.X * ChildAreaRect.Size.Y >= Options.MinimumInteractiveNodeSize );
-		if( Node->bIsInteractive )
+
+		if( Node->bIsInteractive && Node->bIsVisible )
 		{
 			// Figure out how much space we need for the title text
 			const TSharedRef< FSlateFontMeasure >& FontMeasureService = FSlateApplication::Get().GetRenderer()->GetFontMeasureService();
@@ -679,6 +683,11 @@ TArray<FTreeMapNodeVisualInfo> FTreeMap::GetVisuals()
 	{
 		static void RecursivelyGatherVisuals( TArray<FTreeMapNodeVisualInfo>& VisualsList, const FTreeMapNodeRef& Node )
 		{
+			if (!Node->bIsVisible)
+			{
+				return;
+			}
+
 			// Add a visual for the node that was passed in.  We'll recurse down into children afterwards.
 			FTreeMapNodeVisualInfo Visual;
 			Visual.NodeData = Node->Data.Get();
diff --git a/Engine/Source/Editor/ContentBrowser/ContentBrowser.Build.cs b/Engine/Source/Editor/ContentBrowser/ContentBrowser.Build.cs
--- a/Engine/Source/Editor/ContentBrowser/ContentBrowser.Build.cs
+++ b/Engine/Source/Editor/ContentBrowser/ContentBrowser.Build.cs
@@ -16,9 +16,7 @@ public class ContentBrowser : ModuleRules
                 "MainFrame",
 				"PackagesDialog",
 				"SourceControl",
-				"SourceControlWindows",
-                "ReferenceViewer",
-                "SizeMap",
+				"SourceControlWindows"
 			}
 		);
 
@@ -52,9 +50,7 @@ public class ContentBrowser : ModuleRules
 				"AssetTools",
 				"CollectionManager",
 				"GameProjectGeneration",
-                "MainFrame",
-                "ReferenceViewer",
-                "SizeMap",
+                "MainFrame"
 			}
 		);
 		
diff --git a/Engine/Source/Editor/ContentBrowser/Private/AssetContextMenu.cpp b/Engine/Source/Editor/ContentBrowser/Private/AssetContextMenu.cpp
--- a/Engine/Source/Editor/ContentBrowser/Private/AssetContextMenu.cpp
+++ b/Engine/Source/Editor/ContentBrowser/Private/AssetContextMenu.cpp
@@ -43,9 +43,6 @@
 #include "PropertyEditorModule.h"
 #include "Toolkits/GlobalEditorCommonCommands.h"
 #include "ConsolidateWindow.h"
-#include "ReferenceViewer.h"
-#include "ISizeMapModule.h"
-
 #include "ReferencedAssetsUtils.h"
 #include "Internationalization/PackageLocalizationUtil.h"
 
@@ -860,24 +857,6 @@ bool FAssetContextMenu::AddReferenceMenuOptions(FMenuBuilder& MenuBuilder)
 			FSlateIcon(),
 			FUIAction( FExecuteAction::CreateSP( this, &FAssetContextMenu::ExecuteCopyReference ) )
 			);
-
-		MenuBuilder.AddMenuEntry(
-			LOCTEXT("ReferenceViewer", "Reference Viewer..."),
-			LOCTEXT("ReferenceViewerTooltip", "Shows a graph of references for this asset."),
-			FSlateIcon(),
-			FUIAction(
-				FExecuteAction::CreateSP( this, &FAssetContextMenu::ExecuteShowReferenceViewer )
-				)
-			);
-
-		MenuBuilder.AddMenuEntry(
-			LOCTEXT("SizeMap", "Size Map..."),
-			LOCTEXT("SizeMapTooltip", "Shows an interactive map of the approximate memory used by this asset and everything it references."),
-			FSlateIcon(),
-			FUIAction(
-				FExecuteAction::CreateSP( this, &FAssetContextMenu::ExecuteShowSizeMap )
-				)
-			);
 	}
 	MenuBuilder.EndSection();
 
@@ -1971,34 +1950,6 @@ void FAssetContextMenu::ExecuteMigrateAsset()
 	AssetToolsModule.Get().MigratePackages( PackageNames );
 }
 
-void FAssetContextMenu::ExecuteShowReferenceViewer()
-{
-	TArray<FName> PackageNames;
-	for ( auto AssetIt = SelectedAssets.CreateConstIterator(); AssetIt; ++AssetIt )
-	{
-		PackageNames.Add(AssetIt->PackageName);
-	}
-
-	if ( PackageNames.Num() > 0 )
-	{
-		IReferenceViewerModule::Get().InvokeReferenceViewerTab(PackageNames);
-	}
-}
-
-void FAssetContextMenu::ExecuteShowSizeMap()
-{
-	TArray<FName> PackageNames;
-	for ( auto AssetIt = SelectedAssets.CreateConstIterator(); AssetIt; ++AssetIt )
-	{
-		PackageNames.Add(AssetIt->PackageName);
-	}
-
-	if ( PackageNames.Num() > 0 )
-	{
-		ISizeMapModule::Get().InvokeSizeMapTab(PackageNames);
-	}
-}
-
 void FAssetContextMenu::ExecuteGoToCodeForAsset(UClass* SelectedClass)
 {
 	if (SelectedClass)
diff --git a/Engine/Source/Editor/ContentBrowser/Private/AssetContextMenu.h b/Engine/Source/Editor/ContentBrowser/Private/AssetContextMenu.h
--- a/Engine/Source/Editor/ContentBrowser/Private/AssetContextMenu.h
+++ b/Engine/Source/Editor/ContentBrowser/Private/AssetContextMenu.h
@@ -211,12 +211,6 @@ private:
 	/** Handler for when "Migrate Asset" is selected */
 	void ExecuteMigrateAsset();
 
-	/** Handler for ShowReferenceViewer */
-	void ExecuteShowReferenceViewer();
-
-	/** Handler for ShowSizeMap */
-	void ExecuteShowSizeMap();
-
 	/** Handler for GoToAssetCode */
 	void ExecuteGoToCodeForAsset(UClass* SelectedClass);
 
diff --git a/Engine/Source/Editor/ContentBrowser/Private/PathContextMenu.cpp b/Engine/Source/Editor/ContentBrowser/Private/PathContextMenu.cpp
--- a/Engine/Source/Editor/ContentBrowser/Private/PathContextMenu.cpp
+++ b/Engine/Source/Editor/ContentBrowser/Private/PathContextMenu.cpp
@@ -31,8 +31,6 @@
 #include "ContentBrowserUtils.h"
 #include "SourceControlWindows.h"
 #include "ContentBrowserModule.h"
-#include "ReferenceViewer.h"
-#include "ISizeMapModule.h"
 #include "Framework/MultiBox/MultiBoxBuilder.h"
 #include "Widgets/Colors/SColorPicker.h"
 #include "Framework/Commands/GenericCommands.h"
@@ -276,22 +274,6 @@ void FPathContextMenu::MakePathViewContextMenu(FMenuBuilder& MenuBuilder)
 					FSlateIcon()
 					);
 
-				// Reference Viewer
-				MenuBuilder.AddMenuEntry(
-					LOCTEXT("ReferenceViewer", "Reference Viewer..."),
-					LOCTEXT("ReferenceViewerOnFolderTooltip", "Shows a graph of references for this folder."),
-					FSlateIcon(),
-					FUIAction( FExecuteAction::CreateSP( this, &FPathContextMenu::ExecuteReferenceViewer ) )
-					);
-    
-				// Size Map
-				MenuBuilder.AddMenuEntry(
-					LOCTEXT("SizeMap", "Size Map..."),
-					LOCTEXT("SizeMapOnFolderTooltip", "Shows an interactive map of the approximate memory used by the assets in this folder and everything they reference."),
-					FSlateIcon(),
-					FUIAction( FExecuteAction::CreateSP( this, &FPathContextMenu::ExecuteSizeMap ) )
-					);
-
 				// Fix Up Redirectors in Folder
 				MenuBuilder.AddMenuEntry(
 					LOCTEXT("FixUpRedirectorsInFolder", "Fix Up Redirectors in Folder"),
@@ -727,40 +709,6 @@ void FPathContextMenu::ExecuteDelete()
 	}
 }
 
-void FPathContextMenu::ExecuteReferenceViewer()
-{
-	TArray<FString> PackageNamesAsStrings;
-	GetPackageNamesInSelectedPaths(PackageNamesAsStrings);
-
-	TArray<FName> PackageNames;
-	for ( auto PackageNameIt = PackageNamesAsStrings.CreateConstIterator(); PackageNameIt; ++PackageNameIt )
-	{
-		PackageNames.Add(**PackageNameIt);
-	}
-
-	if ( PackageNames.Num() > 0 )
-	{
-		IReferenceViewerModule::Get().InvokeReferenceViewerTab(PackageNames);
-	}
-}
-
-void FPathContextMenu::ExecuteSizeMap()
-{
-	TArray<FString> PackageNamesAsStrings;
-	GetPackageNamesInSelectedPaths(PackageNamesAsStrings);
-
-	TArray<FName> PackageNames;
-	for ( auto PackageNameIt = PackageNamesAsStrings.CreateConstIterator(); PackageNameIt; ++PackageNameIt )
-	{
-		PackageNames.Add(**PackageNameIt);
-	}
-
-	if ( PackageNames.Num() > 0 )
-	{
-		ISizeMapModule::Get().InvokeSizeMapTab(PackageNames);
-	}
-}
-
 void FPathContextMenu::ExecuteFixUpRedirectorsInFolder()
 {
 	FAssetRegistryModule& AssetRegistryModule = FModuleManager::Get().LoadModuleChecked<FAssetRegistryModule>(TEXT("AssetRegistry"));
diff --git a/Engine/Source/Editor/ContentBrowser/Private/PathContextMenu.h b/Engine/Source/Editor/ContentBrowser/Private/PathContextMenu.h
--- a/Engine/Source/Editor/ContentBrowser/Private/PathContextMenu.h
+++ b/Engine/Source/Editor/ContentBrowser/Private/PathContextMenu.h
@@ -95,12 +95,6 @@ public:
 	/** Handler for when "Resave" is selected */
 	void ExecuteResaveFolder();
 
-	/** Handler for when "ReferenceViewer" is selected */
-	void ExecuteReferenceViewer();
-
-	/** Handler for when "SizeMap" is selected */
-	void ExecuteSizeMap();
-
 	/** Handler for when "Fix up Redirectors in Folder" is selected */
 	void ExecuteFixUpRedirectorsInFolder();
 
diff --git a/Engine/Source/Editor/ContentBrowser/Private/SAssetDialog.cpp b/Engine/Source/Editor/ContentBrowser/Private/SAssetDialog.cpp
--- a/Engine/Source/Editor/ContentBrowser/Private/SAssetDialog.cpp
+++ b/Engine/Source/Editor/ContentBrowser/Private/SAssetDialog.cpp
@@ -32,7 +32,6 @@
 #include "Editor/EditorEngine.h"
 #include "FileManager.h"
 #include "NativeClassHierarchy.h"
-#include "ISizeMapModule.h"
 #include "CoreMinimal.h"
 #include "SourceCodeNavigation.h"
 #include "Editor.h"
@@ -545,67 +544,6 @@ void SAssetDialog::ExecuteExplore()
 	}
 }
 
-void SAssetDialog::ExecuteSizeMap()
-{
-	TArray<FString> SelectedPaths = AssetPicker->GetAssetView()->GetSelectedFolders();
-	TArray<FAssetData> SelectedAssets = AssetPicker->GetAssetView()->GetSelectedAssets();
-
-	TArray<FName> PackageNames;
-
-	if (SelectedPaths.Num() == 0 && SelectedAssets.Num() == 0)
-	{
-		SelectedPaths = PathPicker->GetPaths();
-	}
-
-	if (SelectedPaths.Num() > 0)
-	{
-		// Form a filter from the paths
-		FARFilter Filter;
-		Filter.bRecursivePaths = true;
-		for (int32 PathIdx = 0; PathIdx < SelectedPaths.Num(); ++PathIdx)
-		{
-			const FString& Path = SelectedPaths[PathIdx];
-			new (Filter.PackagePaths) FName(*Path);
-		}
-
-		// Load the asset registry module
-		FAssetRegistryModule& AssetRegistryModule = FModuleManager::Get().LoadModuleChecked<FAssetRegistryModule>(TEXT("AssetRegistry"));
-
-		// Query for a list of assets in the selected paths
-		TArray<FAssetData> AssetList;
-		AssetRegistryModule.Get().GetAssets(Filter, AssetList);
-
-		// Form a list of unique package names from the assets
-		TSet<FName> UniquePackageNames;
-		for (int32 AssetIdx = 0; AssetIdx < AssetList.Num(); ++AssetIdx)
-		{
-			UniquePackageNames.Add(AssetList[AssetIdx].PackageName);
-		}
-
-		// Add all unique package names to the output
-		PackageNames.Reserve(UniquePackageNames.Num());
-
-		for (auto PackageIt = UniquePackageNames.CreateConstIterator(); PackageIt; ++PackageIt)
-		{
-			PackageNames.Add((*PackageIt));
-		}
-	}	
-	else
-	{
-		PackageNames.Reserve(SelectedAssets.Num());
-
-		for (auto AssetIt = SelectedAssets.CreateConstIterator(); AssetIt; ++AssetIt)
-		{
-			PackageNames.Add(AssetIt->PackageName);
-		}
-	}
-
-	if (PackageNames.Num() > 0)
-	{
-		ISizeMapModule::Get().InvokeSizeMapModalDialog(PackageNames, FSlateApplication::Get().GetActiveModalWindow());
-	}
-}
-
 bool SAssetDialog::CanExecuteCreateNewFolder() const
 {	
 	// We can only create folders when we have a single path selected
@@ -697,15 +635,6 @@ void SAssetDialog::SetupContextMenuContent(FMenuBuilder& MenuBuilder, const TArr
 							 FUIAction(FExecuteAction::CreateSP(this, &SAssetDialog::ExecuteExplore)));
 
 	MenuBuilder.EndSection();
-
-	MenuBuilder.BeginSection("AssetDialogReferences", LOCTEXT("AssetDialogReferencesHeading", "References"));
-
-	MenuBuilder.AddMenuEntry(LOCTEXT("SizeMap", "Size Map..."), 
-							 LOCTEXT("SizeMapOnFolderTooltip", "Shows an interactive map of the approximate memory used by the assets in this folder and everything they reference."), 
-							 FSlateIcon(), 
-							 FUIAction(FExecuteAction::CreateSP(this, &SAssetDialog::ExecuteSizeMap)));
-
-	MenuBuilder.EndSection();
 }
 
 EActiveTimerReturnType SAssetDialog::SetFocusPostConstruct( double InCurrentTime, float InDeltaTime )
diff --git a/Engine/Source/Editor/ContentBrowser/Private/SAssetDialog.h b/Engine/Source/Editor/ContentBrowser/Private/SAssetDialog.h
--- a/Engine/Source/Editor/ContentBrowser/Private/SAssetDialog.h
+++ b/Engine/Source/Editor/ContentBrowser/Private/SAssetDialog.h
@@ -116,9 +116,6 @@ private:
 	/** Handler for show in explorer */
 	void ExecuteExplore();
 
-	/** Handler for size map tool */
-	void ExecuteSizeMap();
-
 	/** Setup function for the context menu creation of folder and assets */
 	void SetupContextMenuContent(FMenuBuilder& MenuBuilder, const TArray<FString>& SelectedPaths);
 	
diff --git a/Engine/Source/Editor/ContentBrowser/Private/SAssetPicker.cpp b/Engine/Source/Editor/ContentBrowser/Private/SAssetPicker.cpp
--- a/Engine/Source/Editor/ContentBrowser/Private/SAssetPicker.cpp
+++ b/Engine/Source/Editor/ContentBrowser/Private/SAssetPicker.cpp
@@ -244,6 +244,7 @@ void SAssetPicker::Construct( const FArguments& InArgs )
 		.ThumbnailScale(InArgs._AssetPickerConfig.ThumbnailScale)
 		.ShowBottomToolbar(InArgs._AssetPickerConfig.bShowBottomToolbar)
 		.OnAssetTagWantsToBeDisplayed(InArgs._AssetPickerConfig.OnAssetTagWantsToBeDisplayed)
+		.OnGetCustomSourceAssets(InArgs._AssetPickerConfig.OnGetCustomSourceAssets)
 		.AllowDragging( InArgs._AssetPickerConfig.bAllowDragging )
 		.CanShowClasses( InArgs._AssetPickerConfig.bCanShowClasses )
 		.CanShowFolders( InArgs._AssetPickerConfig.bCanShowFolders )
diff --git a/Engine/Source/Editor/ContentBrowser/Private/SAssetView.cpp b/Engine/Source/Editor/ContentBrowser/Private/SAssetView.cpp
--- a/Engine/Source/Editor/ContentBrowser/Private/SAssetView.cpp
+++ b/Engine/Source/Editor/ContentBrowser/Private/SAssetView.cpp
@@ -230,6 +230,7 @@ void SAssetView::Construct( const FArguments& InArgs )
 	OnGetCustomAssetToolTip = InArgs._OnGetCustomAssetToolTip;
 	OnVisualizeAssetToolTip = InArgs._OnVisualizeAssetToolTip;
 	OnAssetToolTipClosing = InArgs._OnAssetToolTipClosing;
+	OnGetCustomSourceAssets = InArgs._OnGetCustomSourceAssets;
 	HighlightedText = InArgs._HighlightedText;
 	ThumbnailLabel = InArgs._ThumbnailLabel;
 	AllowThumbnailHintLabel = InArgs._AllowThumbnailHintLabel;
@@ -1793,6 +1794,12 @@ void SAssetView::RefreshSourceItems()
 				}
 			}
 		}
+
+		// Add any custom assets 
+		if (OnGetCustomSourceAssets.IsBound())
+		{
+			OnGetCustomSourceAssets.Execute(Filter, Items);
+		}
 	}
 
 	// If we are showing classes in the asset list...
diff --git a/Engine/Source/Editor/ContentBrowser/Private/SContentBrowser.cpp b/Engine/Source/Editor/ContentBrowser/Private/SContentBrowser.cpp
--- a/Engine/Source/Editor/ContentBrowser/Private/SContentBrowser.cpp
+++ b/Engine/Source/Editor/ContentBrowser/Private/SContentBrowser.cpp
@@ -55,7 +55,6 @@
 #include "AddToProjectConfig.h"
 #include "GameProjectGenerationModule.h"
 #include "GlobalEditorCommonCommands.h"
-#include "ReferenceViewer.h"
 
 #define LOCTEXT_NAMESPACE "ContentBrowser"
 
@@ -833,10 +832,17 @@ void SContentBrowser::BindCommands()
 		FExecuteAction::CreateSP(this, &SContentBrowser::HandleResaveAllCurrentFolderCommand)
 	));
 
-	Commands->MapAction(FGlobalEditorCommonCommands::Get().ViewReferences, FUIAction(
-		FExecuteAction::CreateSP(this, &SContentBrowser::HandleViewReferencesCommand),
-		FCanExecuteAction::CreateSP(this, &SContentBrowser::HandleViewReferencesCanExecute)
-	));
+	// Allow extenders to add commands
+	FContentBrowserModule& ContentBrowserModule = FModuleManager::GetModuleChecked<FContentBrowserModule>(TEXT("ContentBrowser"));
+	TArray<FContentBrowserCommandExtender> CommmandExtenderDelegates = ContentBrowserModule.GetAllContentBrowserCommandExtenders();
+
+	for (int32 i = 0; i < CommmandExtenderDelegates.Num(); ++i)
+	{
+		if (CommmandExtenderDelegates[i].IsBound())
+		{
+			CommmandExtenderDelegates[i].Execute(Commands.ToSharedRef(), FOnContentBrowserGetSelection::CreateSP(this, &SContentBrowser::GetSelectionState));
+		}
+	}
 }
 
 EVisibility SContentBrowser::GetCollectionViewVisibility() const
@@ -2300,15 +2306,10 @@ void SContentBrowser::HandleDirectoryUpCommandExecute()
 	}
 }
 
-void SContentBrowser::HandleViewReferencesCommand()
+void SContentBrowser::GetSelectionState(TArray<FAssetData>& SelectedAssets, TArray<FString>& SelectedPaths)
 {
-	TArray<FName> ViewableAssets;
-
-	TArray<FAssetData> SelectedAssets;
-	TArray<FString> SelectedPaths;
-
-	// Get the list of selected assets and paths from the view that actually has focus
-
+	SelectedAssets.Reset();
+	SelectedPaths.Reset();
 	if (AssetViewPtr->HasAnyUserFocusOrFocusedDescendants())
 	{
 		SelectedAssets = AssetViewPtr->GetSelectedAssets();
@@ -2318,61 +2319,6 @@ void SContentBrowser::HandleViewReferencesCommand()
 	{
 		SelectedPaths = PathViewPtr->GetSelectedPaths();
 	}
-
-	// For any selected assets, just get the package name from the asset data
-	for (const FAssetData& Asset : SelectedAssets)
-	{
-		ViewableAssets.Add(Asset.PackageName);
-	}
-
-	// For any selected paths, get all assets that exist within that path
-	if (SelectedPaths.Num() > 0)
-	{
-		FARFilter Filter;
-		Filter.bRecursivePaths = true;
-
-		for (const FString& Path : SelectedPaths)
-		{
-			new (Filter.PackagePaths) FName(*Path);
-		}
-
-		FAssetRegistryModule& AssetRegistryModule = FModuleManager::LoadModuleChecked<FAssetRegistryModule>(TEXT("AssetRegistry"));
-		
-		TArray<FAssetData> AssetsInPaths;
-		AssetRegistryModule.Get().GetAssets(Filter, AssetsInPaths);
-
-		for (const FAssetData& Asset : AssetsInPaths)
-		{
-			ViewableAssets.Add(Asset.PackageName);
-		}
-	}
-
-	if (ViewableAssets.Num() > 0)
-	{
-		IReferenceViewerModule::Get().InvokeReferenceViewerTab(ViewableAssets);
-	}
-}
-
-bool SContentBrowser::HandleViewReferencesCanExecute()
-{
-	bool bCanViewReferences = IReferenceViewerModule::IsAvailable();
-
-	if (bCanViewReferences)
-	{
-		// The reference viewer should be called if either the asset view or the path view have focus and have
-		// at least one item selected
-
-		if (AssetViewPtr.IsValid() && AssetViewPtr->HasAnyUserFocusOrFocusedDescendants())
-		{
-			bCanViewReferences = AssetViewPtr->GetSelectedItems().Num() > 0;
-		}
-		else if (PathViewPtr.IsValid() && PathViewPtr->HasAnyUserFocusOrFocusedDescendants())
-		{
-			bCanViewReferences = PathViewPtr->GetSelectedPaths().Num() > 0;
-		}
-	}
-
-	return bCanViewReferences;
 }
 
 bool SContentBrowser::IsBackEnabled() const
diff --git a/Engine/Source/Editor/ContentBrowser/Private/SContentBrowser.h b/Engine/Source/Editor/ContentBrowser/Private/SContentBrowser.h
--- a/Engine/Source/Editor/ContentBrowser/Private/SContentBrowser.h
+++ b/Engine/Source/Editor/ContentBrowser/Private/SContentBrowser.h
@@ -321,12 +321,6 @@ private:
 	/** Handler for clicking the directory up button */
 	void HandleDirectoryUpCommandExecute();
 
-	/** Handler for the view references command */
-	void HandleViewReferencesCommand();
-
-	/** Returns true if the view references command can be handled */
-	bool HandleViewReferencesCanExecute();
-
 	/** True if the user may use the history back button */
 	bool IsBackEnabled() const;
 
@@ -405,6 +399,9 @@ private:
 	/** Delegate called when generating the context menu for a folder */
 	TSharedPtr<SWidget> GetFolderContextMenu(const TArray<FString>& SelectedPaths, FContentBrowserMenuExtender_SelectedPaths InMenuExtender, FOnCreateNewFolder OnCreateNewFolder, bool bPathView);
 
+	/** Delegate called to get the current selection state */
+	void GetSelectionState(TArray<FAssetData>& SelectedAssets, TArray<FString>& SelectedPaths);
+
 	/** Sets up an inline-name for the creation of a default-named folder the specified path */
 	void CreateNewFolder(FString FolderPath, FOnCreateNewFolder OnCreateNewFolder);
 
diff --git a/Engine/Source/Editor/ContentBrowser/Public/ContentBrowserDelegates.h b/Engine/Source/Editor/ContentBrowser/Public/ContentBrowserDelegates.h
--- a/Engine/Source/Editor/ContentBrowser/Public/ContentBrowserDelegates.h
+++ b/Engine/Source/Editor/ContentBrowser/Public/ContentBrowserDelegates.h
@@ -10,6 +10,7 @@
 
 struct FARFilter;
 struct FCollectionNameType;
+class FUICommandList;
 
 
  enum class CONTENTBROWSER_API EMovedContentFolderFlags : uint8
@@ -101,6 +102,10 @@ DECLARE_DELEGATE_OneParam(FOnPathSelected, const FString& /*Path*/);
 /** Called when a path is double clicked in the asset view */
 DECLARE_DELEGATE_OneParam(FOnPathDoubleClicked, const FString& /*Path*/);
 
+/** Called when registering a custom command/keybinding for the content browser */
+DECLARE_DELEGATE_TwoParams(FOnContentBrowserGetSelection, TArray<FAssetData>& /*SelectedAssets*/, TArray<FString>& /*SelectedPaths*/);
+DECLARE_DELEGATE_TwoParams(FContentBrowserCommandExtender, TSharedRef<FUICommandList> /*CommandList*/, FOnContentBrowserGetSelection /*GetSelectionDelegate*/);
+
 /** Called to request the menu when right clicking on a path */
 DECLARE_DELEGATE_RetVal(TSharedRef<FExtender>, FContentBrowserMenuExtender);
 DECLARE_DELEGATE_RetVal_OneParam(TSharedRef<FExtender>, FContentBrowserMenuExtender_SelectedAssets, const TArray<FAssetData>& /*SelectedAssets*/);
@@ -112,9 +117,12 @@ DECLARE_DELEGATE_RetVal_ThreeParams(TSharedPtr<SWidget>, FOnGetFolderContextMenu
 /** Called to request a custom asset item tooltip */
 DECLARE_DELEGATE_RetVal_OneParam( TSharedRef<SToolTip>, FOnGetCustomAssetToolTip, FAssetData& /*AssetData*/);
 
-/** Called to get value for a custom column, will get converted as necesary */
+/** Called to get value for a custom column, will get converted as necessary */
 DECLARE_DELEGATE_RetVal_TwoParams(FString, FOnGetCustomAssetColumnData, FAssetData& /*AssetData*/, FName /*ColumnName*/);
 
+/** Called to add extra asset data to the asset view, to display virtual assets. These get treated similar to Class assets */
+DECLARE_DELEGATE_TwoParams(FOnGetCustomSourceAssets, const FARFilter& /*SourceFilter*/, TArray<FAssetData>& /*AddedAssets*/);
+
 /** Called when an asset item visualizes its tooltip */
 DECLARE_DELEGATE_RetVal_TwoParams(bool, FOnVisualizeAssetToolTip, const TSharedPtr<SWidget>& /*ToolTipContent*/, FAssetData& /*AssetData*/);
 
diff --git a/Engine/Source/Editor/ContentBrowser/Public/ContentBrowserModule.h b/Engine/Source/Editor/ContentBrowser/Public/ContentBrowserModule.h
--- a/Engine/Source/Editor/ContentBrowser/Public/ContentBrowserModule.h
+++ b/Engine/Source/Editor/ContentBrowser/Public/ContentBrowserModule.h
@@ -51,6 +51,9 @@ public:
 	virtual TArray<FContentBrowserMenuExtender_SelectedAssets>& GetAllAssetViewContextMenuExtenders() {return AssetViewContextMenuExtenders;}
 	virtual TArray<FContentBrowserMenuExtender>& GetAllAssetViewViewMenuExtenders() {return AssetViewViewMenuExtenders;}
 
+	/** Delegates to call to extend the command/keybinds for content browser */
+	virtual TArray<FContentBrowserCommandExtender>& GetAllContentBrowserCommandExtenders() { return ContentBrowserCommandExtenders; }
+
 	/** Delegates to be called to extend the drag-and-drop support of the asset view */
 	virtual TArray<FAssetViewDragAndDropExtender>& GetAssetViewDragAndDropExtenders() { return AssetViewDragAndDropExtenders; }
 
@@ -72,6 +75,7 @@ private:
 	TArray<FContentBrowserMenuExtender> CollectionViewContextMenuExtenders;
 	TArray<FContentBrowserMenuExtender_SelectedAssets> AssetViewContextMenuExtenders;
 	TArray<FContentBrowserMenuExtender> AssetViewViewMenuExtenders;
+	TArray<FContentBrowserCommandExtender> ContentBrowserCommandExtenders;
 
 	/** All extender delegates for the drag-and-drop support of the asset view */
 	TArray<FAssetViewDragAndDropExtender> AssetViewDragAndDropExtenders;
diff --git a/Engine/Source/Editor/ContentBrowser/Public/IContentBrowserSingleton.h b/Engine/Source/Editor/ContentBrowser/Public/IContentBrowserSingleton.h
--- a/Engine/Source/Editor/ContentBrowser/Public/IContentBrowserSingleton.h
+++ b/Engine/Source/Editor/ContentBrowser/Public/IContentBrowserSingleton.h
@@ -204,6 +204,9 @@ struct FAssetPickerConfig
 	/** Fired when an asset item is constructed and a tooltip is requested. If unbound the item will use the default widget */
 	FOnGetCustomAssetToolTip OnGetCustomAssetToolTip;
 
+	/** Called to add extra asset data to the asset view, to display virtual assets. These get treated similar to Class assets */
+	FOnGetCustomSourceAssets OnGetCustomSourceAssets;
+
 	/** Fired when an asset item is about to show its tool tip */
 	FOnVisualizeAssetToolTip OnVisualizeAssetToolTip;
 
