Fixing error messaging in collection creation.
#rnx

diff --git a/Engine/Plugins/Editor/AssetManagerEditor/Source/AssetManagerEditor/Private/AssetManagerEditorModule.cpp b/Engine/Plugins/Editor/AssetManagerEditor/Source/AssetManagerEditor/Private/AssetManagerEditorModule.cpp
--- a/Engine/Plugins/Editor/AssetManagerEditor/Source/AssetManagerEditor/Private/AssetManagerEditorModule.cpp
+++ b/Engine/Plugins/Editor/AssetManagerEditor/Source/AssetManagerEditor/Private/AssetManagerEditorModule.cpp
@@ -2149,7 +2149,7 @@ bool FAssetManagerEditorModule::WriteCollection(FName CollectionName, ECollectio
 	else
 	{
 		UE_LOG(LogAssetManagerEditor, Warning, TEXT("Failed to create collection %s. %s"), *CollectionName.ToString(), *CollectionManager.GetLastError().ToString());
-		ResultsMessage = FText::Format(LOCTEXT("CreateCollectionFailed", "Failed to create collection {0}. {0}"), FText::FromName(CollectionName), CollectionManager.GetLastError());
+		ResultsMessage = FText::Format(LOCTEXT("CreateCollectionFailed", "Failed to create collection {0}. {1}"), FText::FromName(CollectionName), CollectionManager.GetLastError());
 	}
 
 	if (bShowFeedback)
