Also clean up the public API a bit

diff --git a/Engine/Plugins/Editor/AssetManagerEditor/Source/AssetManagerEditor/Private/ReferenceViewer/EdGraphNode_Reference.cpp b/Engine/Plugins/Editor/AssetManagerEditor/Source/AssetManagerEditor/Private/ReferenceViewer/EdGraphNode_Reference.cpp
--- a/Engine/Plugins/Editor/AssetManagerEditor/Source/AssetManagerEditor/Private/ReferenceViewer/EdGraphNode_Reference.cpp
+++ b/Engine/Plugins/Editor/AssetManagerEditor/Source/AssetManagerEditor/Private/ReferenceViewer/EdGraphNode_Reference.cpp
@@ -1,6 +1,6 @@
 // Copyright 1998-2018 Epic Games, Inc. All Rights Reserved.
 
-#include "EdGraphNode_Reference.h"
+#include "ReferenceViewer/EdGraphNode_Reference.h"
 #include "GenericPlatform/GenericPlatformFile.h"
 #include "EdGraph/EdGraphPin.h"
 #include "HAL/PlatformFilemanager.h"
@@ -176,8 +176,8 @@ UObject* UEdGraphNode_Reference::GetJumpTargetForDoubleClick() const
 {
 	if (Identifiers.Num() > 0 )
 	{
-		GetReferenceGraph()->SetGraphRoot(Identifiers, FIntPoint(NodePosX, NodePosY));
-		GetReferenceGraph()->RebuildGraph();
+		GetReferenceViewerGraph()->SetGraphRoot(Identifiers, FIntPoint(NodePosX, NodePosY));
+		GetReferenceViewerGraph()->RebuildGraph();
 	}
 	return NULL;
 }
diff --git a/Engine/Plugins/Editor/AssetManagerEditor/Source/AssetManagerEditor/Private/ReferenceViewer/EdGraphNode_Reference.h b/Engine/Plugins/Editor/AssetManagerEditor/Source/AssetManagerEditor/Private/ReferenceViewer/EdGraphNode_Reference.h
deleted file mode 100644
diff --git a/Engine/Plugins/Editor/AssetManagerEditor/Source/AssetManagerEditor/Private/ReferenceViewer/EdGraph_ReferenceViewer.cpp b/Engine/Plugins/Editor/AssetManagerEditor/Source/AssetManagerEditor/Private/ReferenceViewer/EdGraph_ReferenceViewer.cpp
--- a/Engine/Plugins/Editor/AssetManagerEditor/Source/AssetManagerEditor/Private/ReferenceViewer/EdGraph_ReferenceViewer.cpp
+++ b/Engine/Plugins/Editor/AssetManagerEditor/Source/AssetManagerEditor/Private/ReferenceViewer/EdGraph_ReferenceViewer.cpp
@@ -1,8 +1,8 @@
 // Copyright 1998-2018 Epic Games, Inc. All Rights Reserved.
 
-#include "EdGraph_ReferenceViewer.h"
+#include "ReferenceViewer/EdGraph_ReferenceViewer.h"
+#include "ReferenceViewer/EdGraphNode_Reference.h"
 #include "EdGraph/EdGraphPin.h"
-#include "EdGraphNode_Reference.h"
 #include "ARFilter.h"
 #include "AssetRegistryModule.h"
 #include "AssetThumbnail.h"
diff --git a/Engine/Plugins/Editor/AssetManagerEditor/Source/AssetManagerEditor/Private/ReferenceViewer/EdGraph_ReferenceViewer.h b/Engine/Plugins/Editor/AssetManagerEditor/Source/AssetManagerEditor/Private/ReferenceViewer/EdGraph_ReferenceViewer.h
deleted file mode 100644
diff --git a/Engine/Plugins/Editor/AssetManagerEditor/Source/AssetManagerEditor/Private/ReferenceViewer/ReferenceViewerSchema.cpp b/Engine/Plugins/Editor/AssetManagerEditor/Source/AssetManagerEditor/Private/ReferenceViewer/ReferenceViewerSchema.cpp
--- a/Engine/Plugins/Editor/AssetManagerEditor/Source/AssetManagerEditor/Private/ReferenceViewer/ReferenceViewerSchema.cpp
+++ b/Engine/Plugins/Editor/AssetManagerEditor/Source/AssetManagerEditor/Private/ReferenceViewer/ReferenceViewerSchema.cpp
@@ -1,6 +1,6 @@
 // Copyright 1998-2018 Epic Games, Inc. All Rights Reserved.
 
-#include "ReferenceViewerSchema.h"
+#include "ReferenceViewer/ReferenceViewerSchema.h"
 #include "Textures/SlateIcon.h"
 #include "Misc/Attribute.h"
 #include "Framework/MultiBox/MultiBoxBuilder.h"
diff --git a/Engine/Plugins/Editor/AssetManagerEditor/Source/AssetManagerEditor/Private/ReferenceViewer/ReferenceViewerSchema.h b/Engine/Plugins/Editor/AssetManagerEditor/Source/AssetManagerEditor/Private/ReferenceViewer/ReferenceViewerSchema.h
deleted file mode 100644
diff --git a/Engine/Plugins/Editor/AssetManagerEditor/Source/AssetManagerEditor/Private/ReferenceViewer/SReferenceNode.cpp b/Engine/Plugins/Editor/AssetManagerEditor/Source/AssetManagerEditor/Private/ReferenceViewer/SReferenceNode.cpp
--- a/Engine/Plugins/Editor/AssetManagerEditor/Source/AssetManagerEditor/Private/ReferenceViewer/SReferenceNode.cpp
+++ b/Engine/Plugins/Editor/AssetManagerEditor/Source/AssetManagerEditor/Private/ReferenceViewer/SReferenceNode.cpp
@@ -5,7 +5,7 @@
 #include "Widgets/Layout/SSpacer.h"
 #include "Widgets/Images/SImage.h"
 #include "Widgets/Layout/SBox.h"
-#include "EdGraphNode_Reference.h"
+#include "ReferenceViewer/EdGraphNode_Reference.h"
 #include "AssetThumbnail.h"
 #include "Widgets/Text/SInlineEditableTextBlock.h"
 #include "SCommentBubble.h"
diff --git a/Engine/Plugins/Editor/AssetManagerEditor/Source/AssetManagerEditor/Private/ReferenceViewer/SReferenceViewer.cpp b/Engine/Plugins/Editor/AssetManagerEditor/Source/AssetManagerEditor/Private/ReferenceViewer/SReferenceViewer.cpp
--- a/Engine/Plugins/Editor/AssetManagerEditor/Source/AssetManagerEditor/Private/ReferenceViewer/SReferenceViewer.cpp
+++ b/Engine/Plugins/Editor/AssetManagerEditor/Source/AssetManagerEditor/Private/ReferenceViewer/SReferenceViewer.cpp
@@ -18,9 +18,9 @@
 #include "Widgets/Input/SSpinBox.h"
 #include "EditorStyleSet.h"
 #include "Engine/Selection.h"
-#include "EdGraph_ReferenceViewer.h"
-#include "EdGraphNode_Reference.h"
-#include "ReferenceViewerSchema.h"
+#include "ReferenceViewer/EdGraph_ReferenceViewer.h"
+#include "ReferenceViewer/EdGraphNode_Reference.h"
+#include "ReferenceViewer/ReferenceViewerSchema.h"
 #include "AssetRegistryModule.h"
 #include "ICollectionManager.h"
 #include "CollectionManagerModule.h"
diff --git a/Engine/Plugins/Editor/AssetManagerEditor/Source/AssetManagerEditor/Public/ReferenceViewer/EdGraphNode_Reference.h b/Engine/Plugins/Editor/AssetManagerEditor/Source/AssetManagerEditor/Public/ReferenceViewer/EdGraphNode_Reference.h
new file mode 100644
--- /dev/null
+++ b/Engine/Plugins/Editor/AssetManagerEditor/Source/AssetManagerEditor/Public/ReferenceViewer/EdGraphNode_Reference.h
@@ -0,0 +1,68 @@
+// Copyright 1998-2018 Epic Games, Inc. All Rights Reserved.
+
+#pragma once
+
+#include "CoreMinimal.h"
+#include "UObject/ObjectMacros.h"
+#include "EdGraph/EdGraphNode.h"
+#include "AssetData.h"
+#include "EdGraph_ReferenceViewer.h"
+#include "EdGraphNode_Reference.generated.h"
+
+class UEdGraphPin;
+
+UCLASS()
+class ASSETMANAGEREDITOR_API UEdGraphNode_Reference : public UEdGraphNode
+{
+	GENERATED_UCLASS_BODY()
+
+	/** Returns first asset identifier */
+	FAssetIdentifier GetIdentifier() const;
+	
+	/** Returns all identifiers on this node including virtual things */
+	void GetAllIdentifiers(TArray<FAssetIdentifier>& OutIdentifiers) const;
+
+	/** Returns only the packages in this node, skips searchable names */
+	void GetAllPackageNames(TArray<FName>& OutPackageNames) const;
+
+	/** Returns our owning graph */
+	UEdGraph_ReferenceViewer* GetReferenceViewerGraph() const;
+
+	// UEdGraphNode implementation
+	virtual FText GetNodeTitle(ENodeTitleType::Type TitleType) const override;
+	FLinearColor GetNodeTitleColor() const override;
+	virtual FText GetTooltipText() const override;
+	virtual void AllocateDefaultPins() override;
+	virtual UObject* GetJumpTargetForDoubleClick() const override;
+	// End UEdGraphNode implementation
+
+	bool UsesThumbnail() const;
+	bool IsPackage() const;
+	bool IsCollapsed() const;
+	FAssetData GetAssetData() const;
+
+	UEdGraphPin* GetDependencyPin();
+	UEdGraphPin* GetReferencerPin();
+
+private:
+	void CacheAssetData(const FAssetData& AssetData);
+	void SetupReferenceNode(const FIntPoint& NodeLoc, const TArray<FAssetIdentifier>& NewIdentifiers, const FAssetData& InAssetData);
+	void SetReferenceNodeCollapsed(const FIntPoint& NodeLoc, int32 InNumReferencesExceedingMax);
+	void AddReferencer(class UEdGraphNode_Reference* ReferencerNode);
+
+	TArray<FAssetIdentifier> Identifiers;
+	FText NodeTitle;
+
+	bool bUsesThumbnail;
+	bool bIsPackage;
+	bool bIsPrimaryAsset;
+	bool bIsCollapsed;
+	FAssetData CachedAssetData;
+
+	UEdGraphPin* DependencyPin;
+	UEdGraphPin* ReferencerPin;
+
+	friend UEdGraph_ReferenceViewer;
+};
+
+
diff --git a/Engine/Plugins/Editor/AssetManagerEditor/Source/AssetManagerEditor/Public/ReferenceViewer/EdGraph_ReferenceViewer.h b/Engine/Plugins/Editor/AssetManagerEditor/Source/AssetManagerEditor/Public/ReferenceViewer/EdGraph_ReferenceViewer.h
new file mode 100644
--- /dev/null
+++ b/Engine/Plugins/Editor/AssetManagerEditor/Source/AssetManagerEditor/Public/ReferenceViewer/EdGraph_ReferenceViewer.h
@@ -0,0 +1,114 @@
+// Copyright 1998-2018 Epic Games, Inc. All Rights Reserved.
+
+#pragma once
+
+#include "CoreMinimal.h"
+#include "UObject/ObjectMacros.h"
+#include "AssetData.h"
+#include "EdGraph/EdGraph.h"
+#include "Misc/AssetRegistryInterface.h"
+#include "EdGraph_ReferenceViewer.generated.h"
+
+class FAssetThumbnailPool;
+class UEdGraphNode_Reference;
+class SReferenceViewer;
+
+UCLASS()
+class ASSETMANAGEREDITOR_API UEdGraph_ReferenceViewer : public UEdGraph
+{
+	GENERATED_UCLASS_BODY()
+
+public:
+	// UObject implementation
+	virtual void BeginDestroy() override;
+	// End UObject implementation
+
+	/** Set reference viewer to focus on these assets */
+	void SetGraphRoot(const TArray<FAssetIdentifier>& GraphRootIdentifiers, const FIntPoint& GraphRootOrigin = FIntPoint(ForceInitToZero));
+
+	/** Returns list of currently focused assets */
+	const TArray<FAssetIdentifier>& GetCurrentGraphRootIdentifiers() const;
+
+	/** If you're extending the reference viewer via GetAllGraphEditorContextMenuExtender you can use this to get the list of selected assets to use in your menu extender */
+	bool GetSelectedAssetsForMenuExtender(const class UEdGraphNode* Node, TArray<FAssetIdentifier>& SelectedAssets) const;
+
+	/** Accessor for the thumbnail pool in this graph */
+	const TSharedPtr<class FAssetThumbnailPool>& GetAssetThumbnailPool() const;
+
+	/** Force the graph to rebuild */
+	class UEdGraphNode_Reference* RebuildGraph();
+
+	bool IsSearchDepthLimited() const;
+	bool IsSearchBreadthLimited() const;
+	bool IsShowSoftReferences() const;
+	bool IsShowHardReferences() const;
+	bool IsShowManagementReferences() const;
+	bool IsShowSearchableNames() const;
+	bool IsShowNativePackages() const;
+
+	void SetSearchDepthLimitEnabled(bool newEnabled);
+	void SetSearchBreadthLimitEnabled(bool newEnabled);
+	void SetShowSoftReferencesEnabled(bool newEnabled);
+	void SetShowHardReferencesEnabled(bool newEnabled);
+	void SetShowManagementReferencesEnabled(bool newEnabled);
+	void SetShowSearchableNames(bool newEnabled);
+	void SetShowNativePackages(bool newEnabled);
+
+	int32 GetSearchDepthLimit() const;
+	int32 GetSearchBreadthLimit() const;
+
+	void SetSearchDepthLimit(int32 NewDepthLimit);
+	void SetSearchBreadthLimit(int32 NewBreadthLimit);
+
+	FName GetCurrentCollectionFilter() const;
+	void SetCurrentCollectionFilter(FName NewFilter);
+
+	bool GetEnableCollectionFilter() const;
+	void SetEnableCollectionFilter(bool bEnabled);
+
+private:
+	void SetReferenceViewer(TSharedPtr<SReferenceViewer> InViewer);
+	UEdGraphNode_Reference* ConstructNodes(const TArray<FAssetIdentifier>& GraphRootIdentifiers, const FIntPoint& GraphRootOrigin);
+	int32 RecursivelyGatherSizes(bool bReferencers, const TArray<FAssetIdentifier>& Identifiers, const TSet<FName>& AllowedPackageNames, int32 CurrentDepth, TSet<FAssetIdentifier>& VisitedNames, TMap<FAssetIdentifier, int32>& OutNodeSizes) const;
+	void GatherAssetData(const TSet<FName>& AllPackageNames, TMap<FName, FAssetData>& OutPackageToAssetDataMap) const;
+	class UEdGraphNode_Reference* RecursivelyConstructNodes(bool bReferencers, UEdGraphNode_Reference* RootNode, const TArray<FAssetIdentifier>& Identifiers, const FIntPoint& NodeLoc, const TMap<FAssetIdentifier, int32>& NodeSizes, const TMap<FName, FAssetData>& PackagesToAssetDataMap, const TSet<FName>& AllowedPackageNames, int32 CurrentDepth, TSet<FAssetIdentifier>& VisitedNames);
+
+	bool ExceedsMaxSearchDepth(int32 Depth) const;
+	bool ExceedsMaxSearchBreadth(int32 Breadth) const;
+	EAssetRegistryDependencyType::Type GetReferenceSearchFlags(bool bHardOnly) const;
+
+	UEdGraphNode_Reference* CreateReferenceNode();
+
+	/** Removes all nodes from the graph */
+	void RemoveAllNodes();
+
+	/** Returns true if filtering is enabled and we have a valid collection */
+	bool ShouldFilterByCollection() const;
+
+private:
+	/** Pool for maintaining and rendering thumbnails */
+	TSharedPtr<FAssetThumbnailPool> AssetThumbnailPool;
+
+	/** Editor for this pool */
+	TWeakPtr<SReferenceViewer> ReferenceViewer;
+
+	TArray<FAssetIdentifier> CurrentGraphRootIdentifiers;
+	FIntPoint CurrentGraphRootOrigin;
+
+	int32 MaxSearchDepth;
+	int32 MaxSearchBreadth;
+
+	/** Current collection filter. NAME_None for no filter */
+	FName CurrentCollectionFilter;
+	bool bEnableCollectionFilter;
+
+	bool bLimitSearchDepth;
+	bool bLimitSearchBreadth;
+	bool bIsShowSoftReferences;
+	bool bIsShowHardReferences;
+	bool bIsShowManagementReferences;
+	bool bIsShowSearchableNames;
+	bool bIsShowNativePackages;
+
+	friend SReferenceViewer;
+};
diff --git a/Engine/Plugins/Editor/AssetManagerEditor/Source/AssetManagerEditor/Public/ReferenceViewer/ReferenceViewerSchema.h b/Engine/Plugins/Editor/AssetManagerEditor/Source/AssetManagerEditor/Public/ReferenceViewer/ReferenceViewerSchema.h
new file mode 100644
--- /dev/null
+++ b/Engine/Plugins/Editor/AssetManagerEditor/Source/AssetManagerEditor/Public/ReferenceViewer/ReferenceViewerSchema.h
@@ -0,0 +1,40 @@
+// Copyright 1998-2018 Epic Games, Inc. All Rights Reserved.
+
+#pragma once
+
+#include "CoreMinimal.h"
+#include "UObject/ObjectMacros.h"
+#include "EdGraph/EdGraphSchema.h"
+#include "ReferenceViewerSchema.generated.h"
+
+class FMenuBuilder;
+class FSlateRect;
+class UEdGraph;
+
+UCLASS(MinimalAPI)
+class UReferenceViewerSchema : public UEdGraphSchema
+{
+	GENERATED_UCLASS_BODY()
+
+public:
+	// UEdGraphSchema interface
+	virtual void GetContextMenuActions(const UEdGraph* CurrentGraph, const UEdGraphNode* InGraphNode, const UEdGraphPin* InGraphPin, class FMenuBuilder* MenuBuilder, bool bIsDebugging) const override;
+	virtual FLinearColor GetPinTypeColor(const FEdGraphPinType& PinType) const override;
+	virtual void BreakPinLinks(UEdGraphPin& TargetPin, bool bSendsNodeNotifcation) const override;
+	virtual void BreakSinglePinLink(UEdGraphPin* SourcePin, UEdGraphPin* TargetPin) const override;
+	virtual FPinConnectionResponse MovePinLinks(UEdGraphPin& MoveFromPin, UEdGraphPin& MoveToPin, bool bIsIntermediateMove = false, bool bNotifyLinkedNodes = false) const override;
+	virtual FPinConnectionResponse CopyPinLinks(UEdGraphPin& CopyFromPin, UEdGraphPin& CopyToPin, bool bIsIntermediateCopy = false) const override;
+	virtual class FConnectionDrawingPolicy* CreateConnectionDrawingPolicy(int32 InBackLayerID, int32 InFrontLayerID, float InZoomFactor, const FSlateRect& InClippingRect, class FSlateWindowElementList& InDrawElements, class UEdGraph* InGraphObj) const override;
+	virtual void DroppedAssetsOnGraph(const TArray<FAssetData>& Assets, const FVector2D& GraphPosition, UEdGraph* Graph) const override;
+	virtual void GetAssetsGraphHoverMessage(const TArray<FAssetData>& Assets, const UEdGraph* HoverGraph, FString& OutTooltipText, bool& OutOkIcon) const override;
+	// End of UEdGraphSchema interface
+
+private:
+
+	/** Constructs the sub-menu for Make Collection With */
+	void GetMakeCollectionWithSubMenu(FMenuBuilder& MenuBuilder);
+
+	/** Constructs the sub-menu for Make Collection With Referencers/Dependencies */
+	void GetMakeCollectionWithReferencersOrDependenciesSubMenu(class FMenuBuilder& MenuBuilder, bool bReferencers);
+};
+
