PR #6036 : Fix Unreal Engine crash on null pointer

diff --git a/Engine/Plugins/Editor/AssetManagerEditor/Source/AssetManagerEditor/Private/AssetManagerEditorModule.cpp b/Engine/Plugins/Editor/AssetManagerEditor/Source/AssetManagerEditor/Private/AssetManagerEditorModule.cpp
--- a/Engine/Plugins/Editor/AssetManagerEditor/Source/AssetManagerEditor/Private/AssetManagerEditorModule.cpp
+++ b/Engine/Plugins/Editor/AssetManagerEditor/Source/AssetManagerEditor/Private/AssetManagerEditorModule.cpp
@@ -1557,12 +1557,15 @@ void FAssetManagerEditorModule::RefreshRegistryData()
 	UAssetManager::Get().UpdateManagementDatabase(true);
 
 	// Rescan registry sources, try to restore the current one
-	FString OldSourceName = CurrentRegistrySource->SourceName;
+	const FString OldSourceName = CurrentRegistrySource ? CurrentRegistrySource->SourceName : FString();
 
 	CurrentRegistrySource = nullptr;
 	InitializeRegistrySources(false);
 
-	SetCurrentRegistrySource(OldSourceName);
+	if (!OldSourceName.IsEmpty())
+	{
+		SetCurrentRegistrySource(OldSourceName);
+	}
 }
 
 bool FAssetManagerEditorModule::IsPackageInCurrentRegistrySource(FName PackageName)
