fix OpenShaderCookStatistics menu not filtering correctly

diff --git a/Engine/Plugins/Editor/AssetManagerEditor/Source/AssetManagerEditor/Private/AssetManagerEditorModule.cpp b/Engine/Plugins/Editor/AssetManagerEditor/Source/AssetManagerEditor/Private/AssetManagerEditorModule.cpp
--- a/Engine/Plugins/Editor/AssetManagerEditor/Source/AssetManagerEditor/Private/AssetManagerEditorModule.cpp
+++ b/Engine/Plugins/Editor/AssetManagerEditor/Source/AssetManagerEditor/Private/AssetManagerEditorModule.cpp
@@ -732,11 +732,11 @@ void FAssetManagerEditorModule::OpenSizeMapUI(TArray<FAssetIdentifier> SelectedI
 void FAssetManagerEditorModule::OpenShaderCookStatistics(TArray<FName> SelectedPackages)
 {
 	FString SubPath;
-	FString CommonSubPath = "";
+	FString CommonPath = "";
 	if(SelectedPackages.Num())
 	{
 		//Find the common path
-		FString CommonPath = SelectedPackages[0].ToString();
+		CommonPath = SelectedPackages[0].ToString();
 		uint32 CommonIdentical = CommonPath.Len();
 		for (FName Name : SelectedPackages)
 		{
@@ -761,7 +761,7 @@ void FAssetManagerEditorModule::OpenShaderCookStatistics(TArray<FName> SelectedP
 	TSharedRef<SDockTab> Tab = TabManager->InvokeTab(LevelEditorStatsViewerTab);
 	TSharedRef<SWidget> Content = Tab->GetContent();
 	IStatsViewer* StatsView = (IStatsViewer*)&*Content;
-	StatsView->SwitchAndFilterPage(EStatsPage::ShaderCookerStats, CommonSubPath, FString("Path"));
+	StatsView->SwitchAndFilterPage(EStatsPage::ShaderCookerStats, CommonPath, FString("Path"));
 }
 
 void FAssetManagerEditorModule::GetAssetDataInPaths(const TArray<FString>& Paths, TArray<FAssetData>& OutAssetData)
