#UE4 Fix a bug where writing a collection with package names would incorrectly assume the name of the assets in the package.

diff --git a/Engine/Plugins/Editor/AssetManagerEditor/Source/AssetManagerEditor/Private/AssetManagerEditorModule.cpp b/Engine/Plugins/Editor/AssetManagerEditor/Source/AssetManagerEditor/Private/AssetManagerEditorModule.cpp
--- a/Engine/Plugins/Editor/AssetManagerEditor/Source/AssetManagerEditor/Private/AssetManagerEditorModule.cpp
+++ b/Engine/Plugins/Editor/AssetManagerEditor/Source/AssetManagerEditor/Private/AssetManagerEditorModule.cpp
@@ -2186,11 +2186,15 @@ bool FAssetManagerEditorModule::WriteCollection(FName CollectionName, ECollectio
 	bool bSuccess = false;
 
 	TSet<FName> ObjectPathsToAddToCollection;
-	for (FName PackageToAdd : PackageNames)
+
+	FARFilter Filter;
+	Filter.PackageNames = PackageNames;
+	Filter.bIncludeOnlyOnDiskAssets = true;
+	TArray<FAssetData> AssetsInPackages;
+	AssetRegistry->GetAssets(Filter, AssetsInPackages);
+	for (const FAssetData& AssetData : AssetsInPackages)
 	{
-		const FString PackageString = PackageToAdd.ToString();
-		const FName ObjectPath = *FString::Printf(TEXT("%s.%s"), *PackageString, *FPackageName::GetLongPackageAssetName(PackageString));
-		ObjectPathsToAddToCollection.Add(ObjectPath);
+		ObjectPathsToAddToCollection.Add(AssetData.ObjectPath);
 	}
 
 	if (ObjectPathsToAddToCollection.Num() == 0)
