[NO NON-PII COMMIT MESSAGE]

diff --git a/Engine/Plugins/Editor/AssetManagerEditor/Source/AssetManagerEditor/Private/AssetManagerEditorCommands.cpp b/Engine/Plugins/Editor/AssetManagerEditor/Source/AssetManagerEditor/Private/AssetManagerEditorCommands.cpp
--- a/Engine/Plugins/Editor/AssetManagerEditor/Source/AssetManagerEditor/Private/AssetManagerEditorCommands.cpp
+++ b/Engine/Plugins/Editor/AssetManagerEditor/Source/AssetManagerEditor/Private/AssetManagerEditorCommands.cpp
@@ -22,6 +22,7 @@ void FAssetManagerEditorCommands::RegisterCommands()
 	UI_COMMAND(ViewAssetAudit, "Audit Assets...", "Opens the Asset Audit UI and displays information about the selected assets", EUserInterfaceActionType::Button, FInputChord(EModifierKey::Shift | EModifierKey::Alt, EKeys::A));
 
 	UI_COMMAND(OpenSelectedInAssetEditor, "Edit...", "Opens the selected asset in the relevent editor.", EUserInterfaceActionType::Button, FInputChord());
+	UI_COMMAND(ZoomToFit, "Zoom to Fit", "Zoom in and center the view on the selected item", EUserInterfaceActionType::Button, FInputChord(EKeys::Home));
 	UI_COMMAND(ReCenterGraph, "Re-Center Graph", "Re-centers the graph on this node, showing all referencers and references for this asset instead", EUserInterfaceActionType::Button, FInputChord());
 	UI_COMMAND(CopyReferencedObjects, "Copy Referenced Objects List", "Copies the list of objects that the selected asset references to the clipboard.", EUserInterfaceActionType::Button, FInputChord());
 	UI_COMMAND(CopyReferencingObjects, "Copy Referencing Objects List", "Copies the list of objects that reference the selected asset to the clipboard.", EUserInterfaceActionType::Button, FInputChord());
diff --git a/Engine/Plugins/Editor/AssetManagerEditor/Source/AssetManagerEditor/Private/AssetManagerEditorCommands.h b/Engine/Plugins/Editor/AssetManagerEditor/Source/AssetManagerEditor/Private/AssetManagerEditorCommands.h
--- a/Engine/Plugins/Editor/AssetManagerEditor/Source/AssetManagerEditor/Private/AssetManagerEditorCommands.h
+++ b/Engine/Plugins/Editor/AssetManagerEditor/Source/AssetManagerEditor/Private/AssetManagerEditorCommands.h
@@ -49,6 +49,9 @@ public:
 	// Adds all referenced objects to asset audit window
 	TSharedPtr<FUICommandInfo> AuditReferencedObjects;
 
+	/** Zoom in to fit the selected objects in the window */
+	TSharedPtr<FUICommandInfo> ZoomToFit;
+	
 	// Creates a new collection with the list of assets that this asset references, user selects which ECollectionShareType to use.
 	TSharedPtr<FUICommandInfo> MakeLocalCollectionWithReferencers;
 	TSharedPtr<FUICommandInfo> MakePrivateCollectionWithReferencers;
diff --git a/Engine/Plugins/Editor/AssetManagerEditor/Source/AssetManagerEditor/Private/ReferenceViewer/ReferenceViewerSchema.cpp b/Engine/Plugins/Editor/AssetManagerEditor/Source/AssetManagerEditor/Private/ReferenceViewer/ReferenceViewerSchema.cpp
--- a/Engine/Plugins/Editor/AssetManagerEditor/Source/AssetManagerEditor/Private/ReferenceViewer/ReferenceViewerSchema.cpp
+++ b/Engine/Plugins/Editor/AssetManagerEditor/Source/AssetManagerEditor/Private/ReferenceViewer/ReferenceViewerSchema.cpp
@@ -62,6 +62,7 @@ void UReferenceViewerSchema::GetContextMenuActions(const UEdGraph* CurrentGraph,
 
 	MenuBuilder->BeginSection(TEXT("Misc"), NSLOCTEXT("ReferenceViewerSchema", "MiscSectionLabel", "Misc"));
 	{
+		MenuBuilder->AddMenuEntry(FAssetManagerEditorCommands::Get().ZoomToFit);
 		MenuBuilder->AddMenuEntry(FAssetManagerEditorCommands::Get().ReCenterGraph);
 		MenuBuilder->AddSubMenu(
 			NSLOCTEXT("ReferenceViewerSchema", "MakeCollectionWithTitle", "Make Collection with"),
diff --git a/Engine/Plugins/Editor/AssetManagerEditor/Source/AssetManagerEditor/Private/ReferenceViewer/SReferenceViewer.cpp b/Engine/Plugins/Editor/AssetManagerEditor/Source/AssetManagerEditor/Private/ReferenceViewer/SReferenceViewer.cpp
--- a/Engine/Plugins/Editor/AssetManagerEditor/Source/AssetManagerEditor/Private/ReferenceViewer/SReferenceViewer.cpp
+++ b/Engine/Plugins/Editor/AssetManagerEditor/Source/AssetManagerEditor/Private/ReferenceViewer/SReferenceViewer.cpp
@@ -15,6 +15,7 @@
 #include "Widgets/Input/SEditableTextBox.h"
 #include "Widgets/Input/SButton.h"
 #include "Widgets/Input/SCheckBox.h"
+#include "Widgets/Input/SSearchBox.h"
 #include "Widgets/Input/SSpinBox.h"
 #include "EditorStyleSet.h"
 #include "Engine/Selection.h"
@@ -179,6 +180,18 @@ void SReferenceViewer::Construct(const FArguments& InArgs)
 				.BorderImage( FEditorStyle::GetBrush("ToolPanel.GroupBorder") )
 				[
 					SNew(SVerticalBox)
+					+SVerticalBox::Slot()
+					.HAlign(HAlign_Fill)
+					.VAlign(VAlign_Center)
+					.Padding(2.f)
+					.AutoHeight()
+					[
+						SNew(SSearchBox)
+						.HintText(LOCTEXT("Search", "Search..."))
+						.ToolTipText(LOCTEXT("SearchTooltip", "Type here to search (pressing Enter zooms to the results)"))
+						.OnTextChanged(this, &SReferenceViewer::HandleOnSearchTextChanged)
+						.OnTextCommitted(this, &SReferenceViewer::HandleOnSearchTextCommitted)
+					]
 
 					+SVerticalBox::Slot()
 					.AutoHeight()
@@ -898,6 +911,11 @@ void SReferenceViewer::RegisterActions()
 	ReferenceViewerActions = MakeShareable(new FUICommandList);
 	FAssetManagerEditorCommands::Register();
 
+	ReferenceViewerActions->MapAction(
+		FAssetManagerEditorCommands::Get().ZoomToFit,
+		FExecuteAction::CreateSP(this, &SReferenceViewer::ZoomToFit),
+		FCanExecuteAction::CreateSP(this, &SReferenceViewer::CanZoomToFit));
+
 	ReferenceViewerActions->MapAction(
 		FGlobalEditorCommonCommands::Get().FindInContentBrowser,
 		FExecuteAction::CreateSP(this, &SReferenceViewer::ShowSelectionInContentBrowser),
@@ -1427,4 +1445,86 @@ void SReferenceViewer::OnInitialAssetRegistrySearchComplete()
 	}
 }
 
+void SReferenceViewer::ZoomToFit()
+{
+	if (GraphEditorPtr.IsValid())
+	{
+		GraphEditorPtr->ZoomToFit(true);
+	}
+}
+
+bool SReferenceViewer::CanZoomToFit() const
+{
+	if (GraphEditorPtr.IsValid())
+	{
+		return true;
+	}
+
+	return false;
+}
+
+void SReferenceViewer::HandleOnSearchTextChanged(const FText& SearchText)
+{
+	if (GraphObj == nullptr || !GraphEditorPtr.IsValid())
+	{
+		return;
+	}
+
+	GraphEditorPtr->ClearSelectionSet();
+
+	if (SearchText.IsEmpty())
+	{
+		return;
+	}
+
+	FString SearchString = SearchText.ToString();
+	TArray<FString> SearchWords;
+	SearchString.ParseIntoArrayWS( SearchWords );
+
+	TArray<UEdGraphNode_Reference*> AllNodes;
+	GraphObj->GetNodesOfClass<UEdGraphNode_Reference>( AllNodes );
+
+	TArray<FName> NodePackageNames;
+	for (UEdGraphNode_Reference* Node : AllNodes)
+	{
+		NodePackageNames.Empty();
+		Node->GetAllPackageNames(NodePackageNames);
+
+		for (const FName& PackageName : NodePackageNames)
+		{
+			// package name must match all words
+			bool bMatch = true;
+			for (const FString& Word : SearchWords)
+			{
+				if (!PackageName.ToString().Contains(Word))
+				{
+					bMatch = false;
+					break;
+				}
+			}
+
+			if (bMatch)
+			{
+				GraphEditorPtr->SetNodeSelection(Node, true);
+				break;
+			}
+		}
+	}
+}
+
+void SReferenceViewer::HandleOnSearchTextCommitted(const FText& SearchText, ETextCommit::Type CommitType)
+{
+	if (!GraphEditorPtr.IsValid())
+	{
+		return;
+	}
+
+	if (CommitType == ETextCommit::OnCleared)
+	{
+		GraphEditorPtr->ClearSelectionSet();
+	}
+		
+	GraphEditorPtr->ZoomToFit(true);
+}
+
 #undef LOCTEXT_NAMESPACE
diff --git a/Engine/Plugins/Editor/AssetManagerEditor/Source/AssetManagerEditor/Private/ReferenceViewer/SReferenceViewer.h b/Engine/Plugins/Editor/AssetManagerEditor/Source/AssetManagerEditor/Private/ReferenceViewer/SReferenceViewer.h
--- a/Engine/Plugins/Editor/AssetManagerEditor/Source/AssetManagerEditor/Private/ReferenceViewer/SReferenceViewer.h
+++ b/Engine/Plugins/Editor/AssetManagerEditor/Source/AssetManagerEditor/Private/ReferenceViewer/SReferenceViewer.h
@@ -130,6 +130,12 @@ private:
 	void ShowReferenceTree();
 	void ViewSizeMap();
 	void ViewAssetAudit();
+	void ZoomToFit();
+	bool CanZoomToFit() const;
+
+	/** Handlers for searching */
+	void HandleOnSearchTextChanged(const FText& SearchText);
+	void HandleOnSearchTextCommitted(const FText& SearchText, ETextCommit::Type CommitType);
 
 	void ReCenterGraphOnNodes(const TSet<UObject*>& Nodes);
 
