Ctrl+F now focuses the Search... box in the Reference Viewer.

diff --git a/Engine/Plugins/Editor/AssetManagerEditor/Source/AssetManagerEditor/Private/AssetManagerEditorCommands.cpp b/Engine/Plugins/Editor/AssetManagerEditor/Source/AssetManagerEditor/Private/AssetManagerEditorCommands.cpp
--- a/Engine/Plugins/Editor/AssetManagerEditor/Source/AssetManagerEditor/Private/AssetManagerEditorCommands.cpp
+++ b/Engine/Plugins/Editor/AssetManagerEditor/Source/AssetManagerEditor/Private/AssetManagerEditorCommands.cpp
@@ -30,6 +30,7 @@ void FAssetManagerEditorCommands::RegisterCommands()
 	UI_COMMAND(ShowReferencingObjects, "Show Referencing Objects List", "Shows a list of objects that reference the selected asset.", EUserInterfaceActionType::Button, FInputChord());
 	UI_COMMAND(ShowReferenceTree, "Show Reference Tree", "Shows a reference tree for the selected asset.", EUserInterfaceActionType::Button, FInputChord());
 	UI_COMMAND(AuditReferencedObjects, "Audit References", "Opens the Asset Audit UI and displays information about all assets referenced by selected asset", EUserInterfaceActionType::Button, FInputChord());
+	UI_COMMAND(Find, "Find", "Find objects in the reference viewer.", EUserInterfaceActionType::Button, FInputChord(EModifierKey::Control, EKeys::F));
 
 	UI_COMMAND(MakeLocalCollectionWithReferencers, "Local", "Local. This collection is only visible to you and is not in source control.", EUserInterfaceActionType::Button, FInputChord());
 	UI_COMMAND(MakePrivateCollectionWithReferencers, "Private", "Private. This collection is only visible to you.", EUserInterfaceActionType::Button, FInputChord());
diff --git a/Engine/Plugins/Editor/AssetManagerEditor/Source/AssetManagerEditor/Private/AssetManagerEditorCommands.h b/Engine/Plugins/Editor/AssetManagerEditor/Source/AssetManagerEditor/Private/AssetManagerEditorCommands.h
--- a/Engine/Plugins/Editor/AssetManagerEditor/Source/AssetManagerEditor/Private/AssetManagerEditorCommands.h
+++ b/Engine/Plugins/Editor/AssetManagerEditor/Source/AssetManagerEditor/Private/AssetManagerEditorCommands.h
@@ -52,6 +52,9 @@ public:
 	/** Zoom in to fit the selected objects in the window */
 	TSharedPtr<FUICommandInfo> ZoomToFit;
 	
+	/** Start finding objects */
+	TSharedPtr<FUICommandInfo> Find;
+
 	// Creates a new collection with the list of assets that this asset references, user selects which ECollectionShareType to use.
 	TSharedPtr<FUICommandInfo> MakeLocalCollectionWithReferencers;
 	TSharedPtr<FUICommandInfo> MakePrivateCollectionWithReferencers;
diff --git a/Engine/Plugins/Editor/AssetManagerEditor/Source/AssetManagerEditor/Private/ReferenceViewer/SReferenceViewer.cpp b/Engine/Plugins/Editor/AssetManagerEditor/Source/AssetManagerEditor/Private/ReferenceViewer/SReferenceViewer.cpp
--- a/Engine/Plugins/Editor/AssetManagerEditor/Source/AssetManagerEditor/Private/ReferenceViewer/SReferenceViewer.cpp
+++ b/Engine/Plugins/Editor/AssetManagerEditor/Source/AssetManagerEditor/Private/ReferenceViewer/SReferenceViewer.cpp
@@ -33,6 +33,7 @@
 #include "Widgets/Input/SComboBox.h"
 #include "HAL/PlatformApplicationMisc.h"
 #include "AssetManagerEditorModule.h"
+#include "Framework/Application/SlateApplication.h"
 
 #include "ObjectTools.h"
 
@@ -186,7 +187,7 @@ void SReferenceViewer::Construct(const FArguments& InArgs)
 					.Padding(2.f)
 					.AutoHeight()
 					[
-						SNew(SSearchBox)
+						SAssignNew(SearchBox, SSearchBox)
 						.HintText(LOCTEXT("Search", "Search..."))
 						.ToolTipText(LOCTEXT("SearchTooltip", "Type here to search (pressing Enter zooms to the results)"))
 						.OnTextChanged(this, &SReferenceViewer::HandleOnSearchTextChanged)
@@ -916,6 +917,10 @@ void SReferenceViewer::RegisterActions()
 		FExecuteAction::CreateSP(this, &SReferenceViewer::ZoomToFit),
 		FCanExecuteAction::CreateSP(this, &SReferenceViewer::CanZoomToFit));
 
+	ReferenceViewerActions->MapAction(
+		FAssetManagerEditorCommands::Get().Find,
+		FExecuteAction::CreateSP(this, &SReferenceViewer::OnFind));
+
 	ReferenceViewerActions->MapAction(
 		FGlobalEditorCommonCommands::Get().FindInContentBrowser,
 		FExecuteAction::CreateSP(this, &SReferenceViewer::ShowSelectionInContentBrowser),
@@ -1463,6 +1468,11 @@ bool SReferenceViewer::CanZoomToFit() const
 	return false;
 }
 
+void SReferenceViewer::OnFind()
+{
+	FSlateApplication::Get().SetKeyboardFocus(SearchBox, EFocusCause::SetDirectly);
+}
+
 void SReferenceViewer::HandleOnSearchTextChanged(const FText& SearchText)
 {
 	if (GraphObj == nullptr || !GraphEditorPtr.IsValid())
diff --git a/Engine/Plugins/Editor/AssetManagerEditor/Source/AssetManagerEditor/Private/ReferenceViewer/SReferenceViewer.h b/Engine/Plugins/Editor/AssetManagerEditor/Source/AssetManagerEditor/Private/ReferenceViewer/SReferenceViewer.h
--- a/Engine/Plugins/Editor/AssetManagerEditor/Source/AssetManagerEditor/Private/ReferenceViewer/SReferenceViewer.h
+++ b/Engine/Plugins/Editor/AssetManagerEditor/Source/AssetManagerEditor/Private/ReferenceViewer/SReferenceViewer.h
@@ -132,6 +132,7 @@ private:
 	void ViewAssetAudit();
 	void ZoomToFit();
 	bool CanZoomToFit() const;
+	void OnFind();
 
 	/** Handlers for searching */
 	void HandleOnSearchTextChanged(const FText& SearchText);
@@ -159,6 +160,7 @@ private:
 	TSharedPtr<SGraphEditor> GraphEditorPtr;
 
 	TSharedPtr<FUICommandList> ReferenceViewerActions;
+	TSharedPtr<SSearchBox> SearchBox;
 
 	UEdGraph_ReferenceViewer* GraphObj;
 
