Pressing Enter in Reference Viewer Search field now re-runs the search.

Moved check for hard references later to avoid doing an unnecessary TArray::Contains.

diff --git a/Engine/Plugins/Editor/AssetManagerEditor/Source/AssetManagerEditor/Private/ReferenceViewer/SReferenceViewer.cpp b/Engine/Plugins/Editor/AssetManagerEditor/Source/AssetManagerEditor/Private/ReferenceViewer/SReferenceViewer.cpp
--- a/Engine/Plugins/Editor/AssetManagerEditor/Source/AssetManagerEditor/Private/ReferenceViewer/SReferenceViewer.cpp
+++ b/Engine/Plugins/Editor/AssetManagerEditor/Source/AssetManagerEditor/Private/ReferenceViewer/SReferenceViewer.cpp
@@ -1533,7 +1533,11 @@ void SReferenceViewer::HandleOnSearchTextCommitted(const FText& SearchText, ETex
 	{
 		GraphEditorPtr->ClearSelectionSet();
 	}
-		
+	else if (CommitType == ETextCommit::OnEnter)
+	{
+		HandleOnSearchTextChanged(SearchBox->GetText());
+	}
+	
 	GraphEditorPtr->ZoomToFit(true);
 }
 
