Enhancement: When right-clicking a Material Function on the Content Browser, added the "Find Materials Using This" option to open the Reference Viewer with the Materials using that/those selected material function(s).

diff --git a/Engine/Plugins/Editor/AssetManagerEditor/Source/AssetManagerEditor/Private/AssetManagerEditorModule.cpp b/Engine/Plugins/Editor/AssetManagerEditor/Source/AssetManagerEditor/Private/AssetManagerEditorModule.cpp
--- a/Engine/Plugins/Editor/AssetManagerEditor/Source/AssetManagerEditor/Private/AssetManagerEditorModule.cpp
+++ b/Engine/Plugins/Editor/AssetManagerEditor/Source/AssetManagerEditor/Private/AssetManagerEditorModule.cpp
@@ -303,8 +303,8 @@ public:
 	virtual void OpenAssetAuditUI(TArray<FAssetData> SelectedAssets) override;
 	virtual void OpenAssetAuditUI(TArray<FAssetIdentifier> SelectedIdentifiers) override;
 	virtual void OpenAssetAuditUI(TArray<FName> SelectedPackages) override;
-	virtual void OpenReferenceViewerUI(TArray<FAssetIdentifier> SelectedIdentifiers) override;
-	virtual void OpenReferenceViewerUI(TArray<FName> SelectedPackages) override;
+	virtual void OpenReferenceViewerUI(const TArray<FAssetIdentifier> SelectedIdentifiers, const FReferenceViewerParams ReferenceViewerParams = FReferenceViewerParams()) override;
+	virtual void OpenReferenceViewerUI(const TArray<FName> SelectedPackages, const FReferenceViewerParams ReferenceViewerParams = FReferenceViewerParams()) override;
 	virtual void OpenSizeMapUI(TArray<FAssetIdentifier> SelectedIdentifiers) override;
 	virtual void OpenSizeMapUI(TArray<FName> SelectedPackages) override;	
 	virtual bool GetStringValueForCustomColumn(const FAssetData& AssetData, FName ColumnName, FString& OutValue) override;
@@ -683,17 +683,17 @@ void FAssetManagerEditorModule::OpenAssetAuditUI(TArray<FName> SelectedPackages)
 	}
 }
 
-void FAssetManagerEditorModule::OpenReferenceViewerUI(TArray<FAssetIdentifier> SelectedIdentifiers)
+void FAssetManagerEditorModule::OpenReferenceViewerUI(const TArray<FAssetIdentifier> SelectedIdentifiers, const FReferenceViewerParams ReferenceViewerParams)
 {
 	if (SelectedIdentifiers.Num() > 0)
 	{
 		TSharedRef<SDockTab> NewTab = FGlobalTabmanager::Get()->InvokeTab(ReferenceViewerTabName);
 		TSharedRef<SReferenceViewer> ReferenceViewer = StaticCastSharedRef<SReferenceViewer>(NewTab->GetContent());
-		ReferenceViewer->SetGraphRootIdentifiers(SelectedIdentifiers);
+		ReferenceViewer->SetGraphRootIdentifiers(SelectedIdentifiers, ReferenceViewerParams);
 	}
 }
 
-void FAssetManagerEditorModule::OpenReferenceViewerUI(TArray<FName> SelectedPackages)
+void FAssetManagerEditorModule::OpenReferenceViewerUI(const TArray<FName> SelectedPackages, const FReferenceViewerParams ReferenceViewerParams)
 {
 	TArray<FAssetIdentifier> Identifiers;
 	for (FName Name : SelectedPackages)
@@ -701,7 +701,7 @@ void FAssetManagerEditorModule::OpenReferenceViewerUI(TArray<FName> SelectedPack
 		Identifiers.Add(FAssetIdentifier(Name));
 	}
 
-	OpenReferenceViewerUI(Identifiers);
+	OpenReferenceViewerUI(Identifiers, ReferenceViewerParams);
 }
 
 void FAssetManagerEditorModule::OpenSizeMapUI(TArray<FName> SelectedPackages)
@@ -885,7 +885,7 @@ TSharedRef<FExtender> FAssetManagerEditorModule::OnExtendAssetEditor(const TShar
 		// It's safe to modify the CommandList here because this is run as the editor UI is created and the payloads are safe
 		CommandList->MapAction(
 			FAssetManagerEditorCommands::Get().ViewReferences,
-			FExecuteAction::CreateRaw(this, &FAssetManagerEditorModule::OpenReferenceViewerUI, PackageNames));
+			FExecuteAction::CreateRaw(this, &FAssetManagerEditorModule::OpenReferenceViewerUI, PackageNames, FReferenceViewerParams()));
 
 		CommandList->MapAction(
 			FAssetManagerEditorCommands::Get().ViewSizeMap,
diff --git a/Engine/Plugins/Editor/AssetManagerEditor/Source/AssetManagerEditor/Private/ReferenceViewer/EdGraph_ReferenceViewer.cpp b/Engine/Plugins/Editor/AssetManagerEditor/Source/AssetManagerEditor/Private/ReferenceViewer/EdGraph_ReferenceViewer.cpp
--- a/Engine/Plugins/Editor/AssetManagerEditor/Source/AssetManagerEditor/Private/ReferenceViewer/EdGraph_ReferenceViewer.cpp
+++ b/Engine/Plugins/Editor/AssetManagerEditor/Source/AssetManagerEditor/Private/ReferenceViewer/EdGraph_ReferenceViewer.cpp
@@ -29,6 +29,8 @@ UEdGraph_ReferenceViewer::UEdGraph_ReferenceViewer(const FObjectInitializer& Obj
 	bIsShowManagementReferences = false;
 	bIsShowSearchableNames = false;
 	bIsShowNativePackages = false;
+	bIsShowReferencers = true;
+	bIsShowDependencies = true;
 }
 
 void UEdGraph_ReferenceViewer::BeginDestroy()
@@ -143,6 +145,16 @@ bool UEdGraph_ReferenceViewer::IsShowNativePackages() const
 	return bIsShowNativePackages;
 }
 
+bool UEdGraph_ReferenceViewer::IsShowReferencers() const
+{
+	return bIsShowReferencers;
+}
+
+bool UEdGraph_ReferenceViewer::IsShowDependencies() const
+{
+	return bIsShowDependencies;
+}
+
 void UEdGraph_ReferenceViewer::SetSearchDepthLimitEnabled(bool newEnabled)
 {
 	bLimitSearchDepth = newEnabled;
@@ -178,6 +190,16 @@ void UEdGraph_ReferenceViewer::SetShowNativePackages(bool newEnabled)
 	bIsShowNativePackages = newEnabled;
 }
 
+void UEdGraph_ReferenceViewer::SetShowReferencers(const bool newEnabled)
+{
+	bIsShowReferencers = newEnabled;
+}
+
+void UEdGraph_ReferenceViewer::SetShowDependencies(const bool newEnabled)
+{
+	bIsShowDependencies = newEnabled;
+}
+
 int32 UEdGraph_ReferenceViewer::GetSearchDepthLimit() const
 {
 	return MaxSearchDepth;
@@ -250,8 +272,11 @@ UEdGraphNode_Reference* UEdGraph_ReferenceViewer::ConstructNodes(const TArray<FA
 {
 	UEdGraphNode_Reference* RootNode = NULL;
 
-	if (GraphRootIdentifiers.Num() > 0 )
+	if (GraphRootIdentifiers.Num() > 0)
 	{
+		// It both were false, nothing (other than the GraphRootIdentifiers) would be displayed
+		check(bIsShowReferencers || bIsShowDependencies);
+
 		TSet<FName> AllowedPackageNames;
 		if (ShouldFilterByCollection())
 		{
@@ -267,13 +292,19 @@ UEdGraphNode_Reference* UEdGraph_ReferenceViewer::ConstructNodes(const TArray<FA
 
 		TMap<FAssetIdentifier, int32> ReferencerNodeSizes;
 		TSet<FAssetIdentifier> VisitedReferencerSizeNames;
-		int32 ReferencerDepth = 1;
-		RecursivelyGatherSizes(/*bReferencers=*/true, GraphRootIdentifiers, AllowedPackageNames, ReferencerDepth, VisitedReferencerSizeNames, ReferencerNodeSizes);
+		if (bIsShowReferencers)
+		{
+			const int32 ReferencerDepth = 1;
+			RecursivelyGatherSizes(/*bReferencers=*/true, GraphRootIdentifiers, AllowedPackageNames, ReferencerDepth, VisitedReferencerSizeNames, ReferencerNodeSizes);
+		}
 
 		TMap<FAssetIdentifier, int32> DependencyNodeSizes;
 		TSet<FAssetIdentifier> VisitedDependencySizeNames;
-		int32 DependencyDepth = 1;
-		RecursivelyGatherSizes(/*bReferencers=*/false, GraphRootIdentifiers, AllowedPackageNames, DependencyDepth, VisitedDependencySizeNames, DependencyNodeSizes);
+		if (bIsShowDependencies)
+		{
+			const int32 DependencyDepth = 1;
+			RecursivelyGatherSizes(/*bReferencers=*/false, GraphRootIdentifiers, AllowedPackageNames, DependencyDepth, VisitedDependencySizeNames, DependencyNodeSizes);
+		}
 
 		TSet<FName> AllPackageNames;
 
@@ -286,14 +317,20 @@ UEdGraphNode_Reference* UEdGraph_ReferenceViewer::ConstructNodes(const TArray<FA
 			}
 		};
 
-		for (const FAssetIdentifier& AssetId : VisitedReferencerSizeNames)
+		if (bIsShowReferencers)
 		{
-			AddPackage(AssetId, AllPackageNames);
+			for (const FAssetIdentifier& AssetId : VisitedReferencerSizeNames)
+			{
+				AddPackage(AssetId, AllPackageNames);
+			}
 		}
 
-		for (const FAssetIdentifier& AssetId : VisitedDependencySizeNames)
+		if (bIsShowDependencies)
 		{
-			AddPackage(AssetId, AllPackageNames);
+			for (const FAssetIdentifier& AssetId : VisitedDependencySizeNames)
+			{
+				AddPackage(AssetId, AllPackageNames);
+			}
 		}
 
 		TMap<FName, FAssetData> PackagesToAssetDataMap;
@@ -303,13 +340,19 @@ UEdGraphNode_Reference* UEdGraph_ReferenceViewer::ConstructNodes(const TArray<FA
 		RootNode = CreateReferenceNode();
 		RootNode->SetupReferenceNode(GraphRootOrigin, GraphRootIdentifiers, PackagesToAssetDataMap.FindRef(GraphRootIdentifiers[0].PackageName));
 
-		TSet<FAssetIdentifier> VisitedReferencerNames;
-		int32 VisitedReferencerDepth = 1;
-		RecursivelyConstructNodes(/*bReferencers=*/true, RootNode, GraphRootIdentifiers, GraphRootOrigin, ReferencerNodeSizes, PackagesToAssetDataMap, AllowedPackageNames, VisitedReferencerDepth, VisitedReferencerNames);
+		if (bIsShowReferencers)
+		{
+			TSet<FAssetIdentifier> VisitedReferencerNames;
+			const int32 VisitedReferencerDepth = 1;
+			RecursivelyConstructNodes(/*bReferencers=*/true, RootNode, GraphRootIdentifiers, GraphRootOrigin, ReferencerNodeSizes, PackagesToAssetDataMap, AllowedPackageNames, VisitedReferencerDepth, VisitedReferencerNames);
+		}
 
-		TSet<FAssetIdentifier> VisitedDependencyNames;
-		int32 VisitedDependencyDepth = 1;
-		RecursivelyConstructNodes(/*bReferencers=*/false, RootNode, GraphRootIdentifiers, GraphRootOrigin, DependencyNodeSizes, PackagesToAssetDataMap, AllowedPackageNames, VisitedDependencyDepth, VisitedDependencyNames);
+		if (bIsShowDependencies)
+		{
+			TSet<FAssetIdentifier> VisitedDependencyNames;
+			const int32 VisitedDependencyDepth = 1;
+			RecursivelyConstructNodes(/*bReferencers=*/false, RootNode, GraphRootIdentifiers, GraphRootOrigin, DependencyNodeSizes, PackagesToAssetDataMap, AllowedPackageNames, VisitedDependencyDepth, VisitedDependencyNames);
+		}
 	}
 
 	return RootNode;
diff --git a/Engine/Plugins/Editor/AssetManagerEditor/Source/AssetManagerEditor/Private/ReferenceViewer/SReferenceViewer.cpp b/Engine/Plugins/Editor/AssetManagerEditor/Source/AssetManagerEditor/Private/ReferenceViewer/SReferenceViewer.cpp
--- a/Engine/Plugins/Editor/AssetManagerEditor/Source/AssetManagerEditor/Private/ReferenceViewer/SReferenceViewer.cpp
+++ b/Engine/Plugins/Editor/AssetManagerEditor/Source/AssetManagerEditor/Private/ReferenceViewer/SReferenceViewer.cpp
@@ -95,6 +95,14 @@ void SReferenceViewer::Construct(const FArguments& InArgs)
 
 	static const FName DefaultForegroundName("DefaultForeground");
 
+	// Visual options visibility
+	FixAndHideSearchDepthLimit = 0;
+	FixAndHideSearchBreadthLimit = 0;
+	bShowCollectionFilter = true;
+	bShowShowReferencesOptions = true;
+	bShowShowSearchableNames = true;
+	bShowShowNativePackages = true;
+
 	ChildSlot
 	[
 		SNew(SVerticalBox)
@@ -198,6 +206,7 @@ void SReferenceViewer::Construct(const FArguments& InArgs)
 					.AutoHeight()
 					[
 						SNew(SHorizontalBox)
+						.Visibility_Lambda([this]() { return (FixAndHideSearchDepthLimit > 0 ? EVisibility::Collapsed : EVisibility::Visible); })
 
 						+SHorizontalBox::Slot()
 						.VAlign(VAlign_Center)
@@ -239,6 +248,7 @@ void SReferenceViewer::Construct(const FArguments& InArgs)
 					.AutoHeight()
 					[
 						SNew(SHorizontalBox)
+						.Visibility_Lambda([this]() { return (FixAndHideSearchBreadthLimit > 0 ? EVisibility::Collapsed : EVisibility::Visible); })
 
 						+SHorizontalBox::Slot()
 						.VAlign(VAlign_Center)
@@ -280,6 +290,7 @@ void SReferenceViewer::Construct(const FArguments& InArgs)
 					.AutoHeight()
 					[
 						SNew(SHorizontalBox)
+						.Visibility_Lambda([this]() { return (bShowCollectionFilter ? EVisibility::Visible : EVisibility::Collapsed); })
 
 						+SHorizontalBox::Slot()
 						.VAlign(VAlign_Center)
@@ -324,6 +335,7 @@ void SReferenceViewer::Construct(const FArguments& InArgs)
 					.AutoHeight()
 					[
 						SNew(SHorizontalBox)
+						.Visibility_Lambda([this]() { return (bShowShowReferencesOptions ? EVisibility::Visible : EVisibility::Collapsed); })
 
 						+SHorizontalBox::Slot()
 						.VAlign(VAlign_Center)
@@ -349,6 +361,7 @@ void SReferenceViewer::Construct(const FArguments& InArgs)
 					.AutoHeight()
 					[
 						SNew(SHorizontalBox)
+						.Visibility_Lambda([this]() { return (bShowShowReferencesOptions ? EVisibility::Visible : EVisibility::Collapsed); })
 
 						+ SHorizontalBox::Slot()
 						.VAlign(VAlign_Center)
@@ -398,6 +411,7 @@ void SReferenceViewer::Construct(const FArguments& InArgs)
 					.AutoHeight()
 					[
 						SNew(SHorizontalBox)
+						.Visibility_Lambda([this]() { return (bShowShowSearchableNames ? EVisibility::Visible : EVisibility::Collapsed); })
 
 						+ SHorizontalBox::Slot()
 						.VAlign(VAlign_Center)
@@ -422,6 +436,7 @@ void SReferenceViewer::Construct(const FArguments& InArgs)
 					.AutoHeight()
 					[
 						SNew(SHorizontalBox)
+						.Visibility_Lambda([this]() { return (bShowShowNativePackages ? EVisibility::Visible : EVisibility::Collapsed); })
 
 						+ SHorizontalBox::Slot()
 						.VAlign(VAlign_Center)
@@ -455,10 +470,30 @@ void SReferenceViewer::Construct(const FArguments& InArgs)
 	];
 }
 
-void SReferenceViewer::SetGraphRootIdentifiers(const TArray<FAssetIdentifier>& NewGraphRootIdentifiers)
+void SReferenceViewer::SetGraphRootIdentifiers(const TArray<FAssetIdentifier>& NewGraphRootIdentifiers, const FReferenceViewerParams& ReferenceViewerParams)
 {
 	GraphObj->SetGraphRoot(NewGraphRootIdentifiers);
-	
+	// Set properties
+	GraphObj->SetShowReferencers(ReferenceViewerParams.bShowReferencers);
+	GraphObj->SetShowDependencies(ReferenceViewerParams.bShowDependencies);
+	// Set user-interactive properties
+	FixAndHideSearchDepthLimit = ReferenceViewerParams.FixAndHideSearchDepthLimit;
+	if (FixAndHideSearchDepthLimit > 0)
+	{
+		GraphObj->SetSearchDepthLimit(FixAndHideSearchDepthLimit);
+		GraphObj->SetSearchDepthLimitEnabled(true);
+	}
+	FixAndHideSearchBreadthLimit = ReferenceViewerParams.FixAndHideSearchBreadthLimit;
+	if (FixAndHideSearchBreadthLimit > 0)
+	{
+		GraphObj->SetSearchBreadthLimit(FixAndHideSearchBreadthLimit);
+		GraphObj->SetSearchBreadthLimitEnabled(true);
+	}
+	bShowCollectionFilter = ReferenceViewerParams.bShowCollectionFilter;
+	bShowShowReferencesOptions = ReferenceViewerParams.bShowShowReferencesOptions;
+	bShowShowSearchableNames = ReferenceViewerParams.bShowShowSearchableNames;
+	bShowShowNativePackages = ReferenceViewerParams.bShowShowNativePackages;
+
 	RebuildGraph();
 
 	// Zoom once this frame to make sure widgets are visible, then zoom again so size is correct
@@ -813,7 +848,7 @@ ECheckBoxState SReferenceViewer::IsShowHardReferencesChecked() const
 
 EVisibility SReferenceViewer::GetManagementReferencesVisibility() const
 {
-	if (UAssetManager::IsValid())
+	if (bShowShowReferencesOptions && UAssetManager::IsValid())
 	{
 		return EVisibility::SelfHitTestInvisible;
 	}
diff --git a/Engine/Plugins/Editor/AssetManagerEditor/Source/AssetManagerEditor/Private/ReferenceViewer/SReferenceViewer.h b/Engine/Plugins/Editor/AssetManagerEditor/Source/AssetManagerEditor/Private/ReferenceViewer/SReferenceViewer.h
--- a/Engine/Plugins/Editor/AssetManagerEditor/Source/AssetManagerEditor/Private/ReferenceViewer/SReferenceViewer.h
+++ b/Engine/Plugins/Editor/AssetManagerEditor/Source/AssetManagerEditor/Private/ReferenceViewer/SReferenceViewer.h
@@ -31,8 +31,13 @@ public:
 	/** Constructs this widget with InArgs */
 	void Construct( const FArguments& InArgs );
 
-	/** Sets a new root package name */
-	void SetGraphRootIdentifiers(const TArray<FAssetIdentifier>& NewGraphRootIdentifiers);
+	/**
+	 * Sets a new root package name
+	 *
+	 * @param NewGraphRootIdentifiers	The root elements of the new graph to be generated
+	 * @param ReferenceViewerParams		Different visualization settings, such as whether it should display the referencers or the dependencies of the NewGraphRootIdentifiers
+	 */
+	void SetGraphRootIdentifiers(const TArray<FAssetIdentifier>& NewGraphRootIdentifiers, const FReferenceViewerParams& ReferenceViewerParams = FReferenceViewerParams());
 
 	/** Gets graph editor */
 	TSharedPtr<SGraphEditor> GetGraphEditor() const { return GraphEditorPtr; }
@@ -170,4 +175,25 @@ private:
 
 	/** List of collection filter options */
 	TArray<TSharedPtr<FName>> CollectionsComboList;
+
+	/**
+	 * Whether to visually show to the user the option of "Search Depth Limit" or hide it and fix it to a default value:
+	 * - If 0 or negative, it will show to the user the option of "Search Depth Limit".
+	 * - If >0, it will hide that option and fix the Depth value to this value.
+	 */
+	int32 FixAndHideSearchDepthLimit;
+	/**
+	 * Whether to visually show to the user the option of "Search Breadth Limit" or hide it and fix it to a default value:
+	 * - If 0 or negative, it will show to the user the option of "Search Breadth Limit".
+	 * - If >0, it will hide that option and fix the Breadth value to this value.
+	 */
+	int32 FixAndHideSearchBreadthLimit;
+	/** Whether to visually show to the user the option of "Collection Filter" */
+	bool bShowCollectionFilter;
+	/** Whether to visually show to the user the options of "Show Soft/Hard/Management References" */
+	bool bShowShowReferencesOptions;
+	/** Whether to visually show to the user the option of "Show Searchable Names" */
+	bool bShowShowSearchableNames;
+	/** Whether to visually show to the user the option of "Show Native Packages" */
+	bool bShowShowNativePackages;
 };
diff --git a/Engine/Plugins/Editor/AssetManagerEditor/Source/AssetManagerEditor/Public/AssetManagerEditorModule.h b/Engine/Plugins/Editor/AssetManagerEditor/Source/AssetManagerEditor/Public/AssetManagerEditorModule.h
--- a/Engine/Plugins/Editor/AssetManagerEditor/Source/AssetManagerEditor/Public/AssetManagerEditorModule.h
+++ b/Engine/Plugins/Editor/AssetManagerEditor/Source/AssetManagerEditor/Public/AssetManagerEditorModule.h
@@ -190,8 +190,8 @@ public:
 	virtual void OpenAssetAuditUI(TArray<FName> SelectedPackages) = 0;
 
 	/** Spawns reference viewer, showing selected packages or identifiers */
-	virtual void OpenReferenceViewerUI(TArray<FAssetIdentifier> SelectedIdentifiers) = 0;
-	virtual void OpenReferenceViewerUI(TArray<FName> SelectedPackages) = 0;
+	virtual void OpenReferenceViewerUI(const TArray<FAssetIdentifier> SelectedIdentifiers, const FReferenceViewerParams ReferenceViewerParams = FReferenceViewerParams()) = 0;
+	virtual void OpenReferenceViewerUI(const TArray<FName> SelectedPackages, const FReferenceViewerParams ReferenceViewerParams = FReferenceViewerParams()) = 0;
 
 	/** Spawns size map with selected packages */
 	virtual void OpenSizeMapUI(TArray<FAssetIdentifier> SelectedIdentifiers) = 0;
diff --git a/Engine/Plugins/Editor/AssetManagerEditor/Source/AssetManagerEditor/Public/ReferenceViewer/EdGraph_ReferenceViewer.h b/Engine/Plugins/Editor/AssetManagerEditor/Source/AssetManagerEditor/Public/ReferenceViewer/EdGraph_ReferenceViewer.h
--- a/Engine/Plugins/Editor/AssetManagerEditor/Source/AssetManagerEditor/Public/ReferenceViewer/EdGraph_ReferenceViewer.h
+++ b/Engine/Plugins/Editor/AssetManagerEditor/Source/AssetManagerEditor/Public/ReferenceViewer/EdGraph_ReferenceViewer.h
@@ -45,6 +45,8 @@ public:
 	bool IsShowManagementReferences() const;
 	bool IsShowSearchableNames() const;
 	bool IsShowNativePackages() const;
+	bool IsShowReferencers() const;
+	bool IsShowDependencies() const;
 
 	void SetSearchDepthLimitEnabled(bool newEnabled);
 	void SetSearchBreadthLimitEnabled(bool newEnabled);
@@ -53,6 +55,8 @@ public:
 	void SetShowManagementReferencesEnabled(bool newEnabled);
 	void SetShowSearchableNames(bool newEnabled);
 	void SetShowNativePackages(bool newEnabled);
+	void SetShowReferencers(const bool bShouldShowReferencers);
+	void SetShowDependencies(const bool bShouldShowDependencies);
 
 	int32 GetSearchDepthLimit() const;
 	int32 GetSearchBreadthLimit() const;
@@ -109,6 +113,10 @@ private:
 	bool bIsShowManagementReferences;
 	bool bIsShowSearchableNames;
 	bool bIsShowNativePackages;
+	/* Whether to display the Referencers */
+	bool bIsShowReferencers;
+	/* Whether to display the Dependencies */
+	bool bIsShowDependencies;
 
 	friend SReferenceViewer;
 };
diff --git a/Engine/Plugins/Editor/GameplayTagsEditor/Source/GameplayTagsEditor/Private/SGameplayTagWidget.cpp b/Engine/Plugins/Editor/GameplayTagsEditor/Source/GameplayTagsEditor/Private/SGameplayTagWidget.cpp
--- a/Engine/Plugins/Editor/GameplayTagsEditor/Source/GameplayTagsEditor/Private/SGameplayTagWidget.cpp
+++ b/Engine/Plugins/Editor/GameplayTagsEditor/Source/GameplayTagsEditor/Private/SGameplayTagWidget.cpp
@@ -981,7 +981,7 @@ void SGameplayTagWidget::OnSearchForReferences(TSharedPtr<FGameplayTagNode> InTa
 	{
 		TArray<FAssetIdentifier> AssetIdentifiers;
 		AssetIdentifiers.Add(FAssetIdentifier(FGameplayTag::StaticStruct(), InTagNode->GetCompleteTagName()));
-		FEditorDelegates::OnOpenReferenceViewer.Broadcast(AssetIdentifiers);
+		FEditorDelegates::OnOpenReferenceViewer.Broadcast(AssetIdentifiers, FReferenceViewerParams());
 	}
 }
 
diff --git a/Engine/Source/Developer/AssetTools/Private/AssetTypeActions/AssetTypeActions_MaterialFunction.cpp b/Engine/Source/Developer/AssetTools/Private/AssetTypeActions/AssetTypeActions_MaterialFunction.cpp
--- a/Engine/Source/Developer/AssetTools/Private/AssetTypeActions/AssetTypeActions_MaterialFunction.cpp
+++ b/Engine/Source/Developer/AssetTools/Private/AssetTypeActions/AssetTypeActions_MaterialFunction.cpp
@@ -14,6 +14,8 @@
 #include "ContentBrowserModule.h"
 #include "IContentBrowserSingleton.h"
 #include "ARFilter.h"
+#include "AssetData.h"
+#include "Misc/PackageName.h"
 
 #define LOCTEXT_NAMESPACE "AssetTypeActions"
 
@@ -26,7 +28,7 @@ const FString BlendCompareString = (TEXT("MaterialLayerBlend"));
 
 void FAssetTypeActions_MaterialFunction::GetActions(const TArray<UObject*>& InObjects, FToolMenuSection& Section)
 {
-	auto Functions = GetTypedWeakObjectPtrs<UMaterialFunctionInterface>(InObjects);
+	const TArray<TWeakObjectPtr<UMaterialFunctionInterface>> Functions = GetTypedWeakObjectPtrs<UMaterialFunctionInterface>(InObjects);
 
 	IMaterialEditorModule& MaterialEditorModule = FModuleManager::LoadModuleChecked<IMaterialEditorModule>("MaterialEditor");
 	if (MaterialEditorModule.MaterialLayersEnabled())
@@ -42,16 +44,19 @@ void FAssetTypeActions_MaterialFunction::GetActions(const TArray<UObject*>& InOb
 		);
 	}
 
-	Section.AddMenuEntry(
-		"MaterialFunction_FindMaterials",
-		LOCTEXT("MaterialFunction_FindMaterials", "Find Materials Using This"),
-		LOCTEXT("MaterialFunction_FindMaterialsTooltip", "Finds the materials that reference this material function in the content browser."),
-		FSlateIcon(FEditorStyle::GetStyleSetName(), "ContentBrowser.AssetActions.GenericFind"),
-		FUIAction(
-			FExecuteAction::CreateSP( this, &FAssetTypeActions_MaterialFunction::ExecuteFindMaterials, Functions ),
-			FCanExecuteAction()
-			)
-		);
+	if (FEditorDelegates::OnOpenReferenceViewer.IsBound())
+	{
+		Section.AddMenuEntry(
+			"MaterialFunction_FindMaterials",
+			LOCTEXT("MaterialFunction_FindMaterials", "Find Materials Using This"),
+			LOCTEXT("MaterialFunction_FindMaterialsTooltip", "Finds the materials that reference this material function and visually displays them with the light version of the Reference Viewer."),
+			FSlateIcon(FEditorStyle::GetStyleSetName(), "ContentBrowser.AssetActions.GenericFind"),
+			FUIAction(
+				FExecuteAction::CreateSP( this, &FAssetTypeActions_MaterialFunction::ExecuteFindMaterials, Functions ),
+				FCanExecuteAction()
+				)
+			);
+	}
 }
 
 void FAssetTypeActions_MaterialFunction::OpenAssetEditor( const TArray<UObject*>& InObjects, TSharedPtr<IToolkitHost> EditWithinLevelEditor )
@@ -127,34 +132,32 @@ void FAssetTypeActions_MaterialFunction::ExecuteNewMFI(TArray<TWeakObjectPtr<UMa
 
 void FAssetTypeActions_MaterialFunction::ExecuteFindMaterials(TArray<TWeakObjectPtr<UMaterialFunctionInterface>> Objects)
 {
-	TArray<UObject*> ObjectsToSync;
-
-	for (auto ObjIt = Objects.CreateConstIterator(); ObjIt; ++ObjIt)
+#ifdef WITH_EDITOR
+	if (FEditorDelegates::OnOpenReferenceViewer.IsBound())
 	{
-		auto Object = (*ObjIt).Get();
-		if ( Object )
+		// TArray that will be send to the ReferenceViewer for display
+		TArray<FAssetIdentifier> AssetIdentifiers;
+		// Iterate over all selected UMaterialFunctionInterface instances
+		for (auto ObjIt = Objects.CreateConstIterator(); ObjIt; ++ObjIt)
 		{
-			// @todo This only considers loaded materials! Find a good way to make this use the asset registry.
-			for (TObjectIterator<UMaterial> It; It; ++It)
+			const UMaterialFunctionInterface* const Object = (*ObjIt).Get();
+			// If valid pointer, add to AssetIdentifiers for display on the ReferenceViewer
+			if (Object)
 			{
-				UMaterial* CurrentMaterial = *It;
-
-				for (int32 FunctionIndex = 0; FunctionIndex < CurrentMaterial->MaterialFunctionInfos.Num(); FunctionIndex++)
-				{
-					if (CurrentMaterial->MaterialFunctionInfos[FunctionIndex].Function == Object)
-					{
-						ObjectsToSync.Add(CurrentMaterial);
-						break;
-					}
-				}
+				// Construct FAssetIdentifier and add to TArray
+				const FName AssetName(*FPackageName::ObjectPathToPackageName(GetPathNameSafe(Object)));
+				AssetIdentifiers.Add(FAssetIdentifier(AssetName));
 			}
 		}
+		// Call ReferenceViewer
+		FReferenceViewerParams ReferenceViewerParams;
+		ReferenceViewerParams.bShowDependencies = false;
+		ReferenceViewerParams.FixAndHideSearchDepthLimit = 1;
+		ReferenceViewerParams.bShowShowReferencesOptions = false;
+		ReferenceViewerParams.bShowShowSearchableNames = false;
+		FEditorDelegates::OnOpenReferenceViewer.Broadcast(AssetIdentifiers, ReferenceViewerParams);
 	}
-
-	if (ObjectsToSync.Num() > 0)
-	{
-		FAssetTools::Get().SyncBrowserToAssets(ObjectsToSync);
-	}
+#endif // WITH_EDITOR
 }
 
 UThumbnailInfo* FAssetTypeActions_MaterialFunction::GetThumbnailInfo(UObject* Asset) const
diff --git a/Engine/Source/Editor/DataTableEditor/Private/SDataTableListViewRow.cpp b/Engine/Source/Editor/DataTableEditor/Private/SDataTableListViewRow.cpp
--- a/Engine/Source/Editor/DataTableEditor/Private/SDataTableListViewRow.cpp
+++ b/Engine/Source/Editor/DataTableEditor/Private/SDataTableListViewRow.cpp
@@ -59,7 +59,7 @@ void SDataTableListViewRow::OnSearchForReferences()
 			TArray<FAssetIdentifier> AssetIdentifiers;
 			AssetIdentifiers.Add(FAssetIdentifier(SourceDataTable, RowDataPtr->RowId));
 
-			FEditorDelegates::OnOpenReferenceViewer.Broadcast(AssetIdentifiers);
+			FEditorDelegates::OnOpenReferenceViewer.Broadcast(AssetIdentifiers, FReferenceViewerParams());
 		}
 	}
 }
diff --git a/Engine/Source/Editor/DetailCustomizations/Private/DataTableCustomization.cpp b/Engine/Source/Editor/DetailCustomizations/Private/DataTableCustomization.cpp
--- a/Engine/Source/Editor/DetailCustomizations/Private/DataTableCustomization.cpp
+++ b/Engine/Source/Editor/DetailCustomizations/Private/DataTableCustomization.cpp
@@ -143,7 +143,7 @@ void FDataTableCustomizationLayout::OnSearchForReferences()
 		TArray<FAssetIdentifier> AssetIdentifiers;
 		AssetIdentifiers.Add(FAssetIdentifier(SourceDataTable, RowName));
 
-		FEditorDelegates::OnOpenReferenceViewer.Broadcast(AssetIdentifiers);
+		FEditorDelegates::OnOpenReferenceViewer.Broadcast(AssetIdentifiers, FReferenceViewerParams());
 	}
 }
 
diff --git a/Engine/Source/Editor/UnrealEd/Private/Editor.cpp b/Engine/Source/Editor/UnrealEd/Private/Editor.cpp
--- a/Engine/Source/Editor/UnrealEd/Private/Editor.cpp
+++ b/Engine/Source/Editor/UnrealEd/Private/Editor.cpp
@@ -140,7 +140,7 @@ FEditorDelegates::FOnDuplicateActorsBegin				FEditorDelegates::OnDuplicateActors
 FEditorDelegates::FOnDuplicateActorsEnd					FEditorDelegates::OnDuplicateActorsEnd;
 FEditorDelegates::FOnDeleteActorsBegin					FEditorDelegates::OnDeleteActorsBegin;
 FEditorDelegates::FOnDeleteActorsEnd					FEditorDelegates::OnDeleteActorsEnd;
-FEditorDelegates::FOnViewAssetIdentifiers				FEditorDelegates::OnOpenReferenceViewer;
+FEditorDelegates::FOnOpenReferenceViewer				FEditorDelegates::OnOpenReferenceViewer;
 FEditorDelegates::FOnViewAssetIdentifiers				FEditorDelegates::OnOpenSizeMap;
 FEditorDelegates::FOnViewAssetIdentifiers				FEditorDelegates::OnOpenAssetAudit;
 FEditorDelegates::FOnViewAssetIdentifiers				FEditorDelegates::OnEditAssetIdentifiers;
diff --git a/Engine/Source/Editor/UnrealEd/Public/Editor.h b/Engine/Source/Editor/UnrealEd/Public/Editor.h
--- a/Engine/Source/Editor/UnrealEd/Public/Editor.h
+++ b/Engine/Source/Editor/UnrealEd/Public/Editor.h
@@ -134,6 +134,8 @@ struct UNREALED_API FEditorDelegates
 	DECLARE_MULTICAST_DELEGATE(FOnDeleteActorsEnd);
 	/** delegate type to handle viewing/editing a set of asset identifiers which are packages or ids */
 	DECLARE_MULTICAST_DELEGATE_OneParam(FOnViewAssetIdentifiers, TArray<FAssetIdentifier>);
+	/** delegate type to handle viewing/editing a set of asset identifiers (which are packages or ids) in the reference viewer */
+	DECLARE_MULTICAST_DELEGATE_TwoParams(FOnOpenReferenceViewer, const TArray<FAssetIdentifier>, const FReferenceViewerParams);
 
 	/** Called when the CurrentLevel is switched to a new level.  Note that this event won't be fired for temporary
 		changes to the current level, such as when copying/pasting actors. */
@@ -279,7 +281,7 @@ struct UNREALED_API FEditorDelegates
 	/** Sent when delete end called */
 	static FOnDeleteActorsEnd OnDeleteActorsEnd;
 	/** Called when you want to view things in the reference viewer, these are bound to by asset manager editor plugins */
-	static FOnViewAssetIdentifiers OnOpenReferenceViewer;
+	static FOnOpenReferenceViewer OnOpenReferenceViewer;
 	/** Called when you want to view things in the size map */
 	static FOnViewAssetIdentifiers OnOpenSizeMap;
 	/** Called when you want to view things in the asset audit window */
diff --git a/Engine/Source/Runtime/AssetRegistry/Public/AssetData.h b/Engine/Source/Runtime/AssetRegistry/Public/AssetData.h
--- a/Engine/Source/Runtime/AssetRegistry/Public/AssetData.h
+++ b/Engine/Source/Runtime/AssetRegistry/Public/AssetData.h
@@ -598,6 +598,53 @@ public:
 	}
 };
 
+/**
+ * Helper struct for FAssetIdentifier (e.g., for the FOnViewAssetIdentifiersInReferenceViewer delegate and Reference Viewer functions).
+ */
+#ifdef WITH_EDITORONLY_DATA
+struct FReferenceViewerParams
+{
+	FReferenceViewerParams()
+		// Displayed-on-graph options
+		: bShowReferencers(true)
+		, bShowDependencies(true)
+		// Slider-based options
+		, FixAndHideSearchDepthLimit(0)
+		, FixAndHideSearchBreadthLimit(0)
+		, bShowCollectionFilter(true)
+		// Checkbox options
+		, bShowShowReferencesOptions(true)
+		, bShowShowSearchableNames(true)
+		, bShowShowNativePackages(true)
+	{}
+
+	/* Whether to display the Referencers */
+	bool bShowReferencers;
+	/* Whether to display the Dependencies */
+	bool bShowDependencies;
+	/**
+	 * Whether to visually show to the user the option of "Search Depth Limit" or hide it and fix it to a default value:
+	 * - If 0 or negative, it will show to the user the option of "Search Depth Limit".
+	 * - If >0, it will hide that option and fix the Depth value to this value.
+	 */
+	int32 FixAndHideSearchDepthLimit;
+	/**
+	 * Whether to visually show to the user the option of "Search Breadth Limit" or hide it and fix it to a default value:
+	 * - If 0 or negative, it will show to the user the option of "Search Breadth Limit".
+	 * - If >0, it will hide that option and fix the Breadth value to this value.
+	 */
+	int32 FixAndHideSearchBreadthLimit;
+	/** Whether to visually show to the user the option of "Collection Filter" */
+	bool bShowCollectionFilter;
+	/** Whether to visually show to the user the options of "Show Soft/Hard/Management References" */
+	bool bShowShowReferencesOptions;
+	/** Whether to visually show to the user the option of "Show Searchable Names" */
+	bool bShowShowSearchableNames;
+	/** Whether to visually show to the user the option of "Show Native Packages" */
+	bool bShowShowNativePackages;
+};
+#endif // WITH_EDITORONLY_DATA
+
 /** A structure defining a thing that can be reference by something else in the asset registry. Represents either a package of a primary asset id */
 struct FAssetIdentifier
 {
