Bug fix: Reference Viewer showed a +1 offset in the "Search Breadth Limit" option. I.e., if selected a Breadth of n, n+1 elements were shown rather than only n.

diff --git a/Engine/Plugins/Editor/AssetManagerEditor/Source/AssetManagerEditor/Private/ReferenceViewer/EdGraph_ReferenceViewer.cpp b/Engine/Plugins/Editor/AssetManagerEditor/Source/AssetManagerEditor/Private/ReferenceViewer/EdGraph_ReferenceViewer.cpp
--- a/Engine/Plugins/Editor/AssetManagerEditor/Source/AssetManagerEditor/Private/ReferenceViewer/EdGraph_ReferenceViewer.cpp
+++ b/Engine/Plugins/Editor/AssetManagerEditor/Source/AssetManagerEditor/Private/ReferenceViewer/EdGraph_ReferenceViewer.cpp
@@ -620,12 +620,14 @@ const TSharedPtr<FAssetThumbnailPool>& UEdGraph_ReferenceViewer::GetAssetThumbna
 
 bool UEdGraph_ReferenceViewer::ExceedsMaxSearchDepth(int32 Depth) const
 {
+	// ExceedsMaxSearchDepth requires only greater (not equal than) because, even though the Depth is 1-based indexed (similarly to Breadth), the first index (index 0) corresponds to the root object 
 	return bLimitSearchDepth && Depth > MaxSearchDepth;
 }
 
 bool UEdGraph_ReferenceViewer::ExceedsMaxSearchBreadth(int32 Breadth) const
 {
-	return bLimitSearchBreadth && Breadth > MaxSearchBreadth;
+	// ExceedsMaxSearchBreadth requires greater or equal than because the Breadth is 1-based indexed
+	return bLimitSearchBreadth && Breadth >= MaxSearchBreadth;
 }
 
 UEdGraphNode_Reference* UEdGraph_ReferenceViewer::CreateReferenceNode()
