Hide dynamic collections from the reference viewer list as they require frontend filters

Also refresh the combo list when opening the UI to pick up newly added or removed collections

diff --git a/Engine/Plugins/Editor/AssetManagerEditor/Source/AssetManagerEditor/Private/ReferenceViewer/SReferenceViewer.cpp b/Engine/Plugins/Editor/AssetManagerEditor/Source/AssetManagerEditor/Private/ReferenceViewer/SReferenceViewer.cpp
--- a/Engine/Plugins/Editor/AssetManagerEditor/Source/AssetManagerEditor/Private/ReferenceViewer/SReferenceViewer.cpp
+++ b/Engine/Plugins/Editor/AssetManagerEditor/Source/AssetManagerEditor/Private/ReferenceViewer/SReferenceViewer.cpp
@@ -59,19 +59,6 @@ void SReferenceViewer::Construct(const FArguments& InArgs)
 	HistoryManager.SetOnApplyHistoryData(FOnApplyHistoryData::CreateSP(this, &SReferenceViewer::OnApplyHistoryData));
 	HistoryManager.SetOnUpdateHistoryData(FOnUpdateHistoryData::CreateSP(this, &SReferenceViewer::OnUpdateHistoryData));
 
-	// Fill out CollectionsComboList
-	{
-		TArray<FName> CollectionNames;
-		FCollectionManagerModule& CollectionManagerModule = FCollectionManagerModule::GetModule();
-		CollectionManagerModule.Get().GetCollectionNames(ECollectionShareType::CST_All, CollectionNames);
-		CollectionNames.Sort([](const FName& A, const FName& B) { return A.Compare(B) < 0; });
-		CollectionsComboList.Add(MakeShareable(new FName(NAME_None)));
-		for (FName CollectionName : CollectionNames)
-		{
-			CollectionsComboList.Add(MakeShareable(new FName(CollectionName)));
-		}
-	}
-
 	// Create the graph
 	GraphObj = NewObject<UEdGraph_ReferenceViewer>();
 	GraphObj->Schema = UReferenceViewerSchema::StaticClass();
@@ -318,8 +305,9 @@ void SReferenceViewer::Construct(const FArguments& InArgs)
 							SNew(SBox)
 							.WidthOverride(100)
 							[
-								SNew(SComboBox<TSharedPtr<FName>>)
+								SAssignNew(CollectionsCombo, SComboBox<TSharedPtr<FName>>)
 								.OptionsSource(&CollectionsComboList)
+								.OnComboBoxOpening(this, &SReferenceViewer::UpdateCollectionsComboList)
 								.OnGenerateWidget(this, &SReferenceViewer::GenerateCollectionFilterItem)
 								.OnSelectionChanged(this, &SReferenceViewer::HandleCollectionFilterChanged)
 								.ToolTipText(this, &SReferenceViewer::GetCollectionFilterText)
@@ -468,6 +456,8 @@ void SReferenceViewer::Construct(const FArguments& InArgs)
 			]
 		]
 	];
+
+	UpdateCollectionsComboList();
 }
 
 void SReferenceViewer::SetGraphRootIdentifiers(const TArray<FAssetIdentifier>& NewGraphRootIdentifiers, const FReferenceViewerParams& ReferenceViewerParams)
@@ -778,9 +768,60 @@ ECheckBoxState SReferenceViewer::IsEnableCollectionFilterChecked() const
 	}
 }
 
+void SReferenceViewer::UpdateCollectionsComboList()
+{
+	TArray<FName> CollectionNames;
+	{
+		FCollectionManagerModule& CollectionManagerModule = FCollectionManagerModule::GetModule();
+
+		TArray<FCollectionNameType> AllCollections;
+		CollectionManagerModule.Get().GetCollections(AllCollections);
+
+		for (const FCollectionNameType& Collection : AllCollections)
+		{
+			ECollectionStorageMode::Type StorageMode = ECollectionStorageMode::Static;
+			CollectionManagerModule.Get().GetCollectionStorageMode(Collection.Name, Collection.Type, StorageMode);
+
+			if (StorageMode == ECollectionStorageMode::Static)
+			{
+				CollectionNames.AddUnique(Collection.Name);
+			}
+		}
+	}
+	CollectionNames.Sort([](const FName& A, const FName& B) { return A.Compare(B) < 0; });
+
+	CollectionsComboList.Reset();
+	CollectionsComboList.Add(MakeShared<FName>(NAME_None));
+	for (FName CollectionName : CollectionNames)
+	{
+		CollectionsComboList.Add(MakeShared<FName>(CollectionName));
+	}
+
+	if (CollectionsCombo)
+	{
+		CollectionsCombo->ClearSelection();
+		CollectionsCombo->RefreshOptions();
+
+		if (GraphObj)
+		{
+			const FName CurrentFilter = GraphObj->GetCurrentCollectionFilter();
+
+			const int32 SelectedItemIndex = CollectionsComboList.IndexOfByPredicate([CurrentFilter](const TSharedPtr<FName>& InItem)
+			{
+				return CurrentFilter == *InItem;
+			});
+
+			if (SelectedItemIndex != INDEX_NONE)
+			{
+				CollectionsCombo->SetSelectedItem(CollectionsComboList[SelectedItemIndex]);
+			}
+		}
+	}
+}
+
 void SReferenceViewer::HandleCollectionFilterChanged(TSharedPtr<FName> Item, ESelectInfo::Type SelectInfo)
 {
-	if (GraphObj)
+	if (GraphObj && Item)
 	{
 		const FName NewFilter = *Item;
 		const FName CurrentFilter = GraphObj->GetCurrentCollectionFilter();
diff --git a/Engine/Plugins/Editor/AssetManagerEditor/Source/AssetManagerEditor/Private/ReferenceViewer/SReferenceViewer.h b/Engine/Plugins/Editor/AssetManagerEditor/Source/AssetManagerEditor/Private/ReferenceViewer/SReferenceViewer.h
--- a/Engine/Plugins/Editor/AssetManagerEditor/Source/AssetManagerEditor/Private/ReferenceViewer/SReferenceViewer.h
+++ b/Engine/Plugins/Editor/AssetManagerEditor/Source/AssetManagerEditor/Private/ReferenceViewer/SReferenceViewer.h
@@ -104,6 +104,7 @@ private:
 	void OnEnableCollectionFilterChanged(ECheckBoxState NewState);
 	ECheckBoxState IsEnableCollectionFilterChecked() const;
 	TSharedRef<SWidget> GenerateCollectionFilterItem(TSharedPtr<FName> InItem);
+	void UpdateCollectionsComboList();
 	void HandleCollectionFilterChanged(TSharedPtr<FName> Item, ESelectInfo::Type SelectInfo);
 	FText GetCollectionFilterText() const;
 
@@ -173,6 +174,9 @@ private:
 	/** The temporary copy of the path text when it is actively being edited. */
 	FText TemporaryPathBeingEdited;
 
+	/** Combo box for collections filter options */
+	TSharedPtr<SComboBox<TSharedPtr<FName>>> CollectionsCombo;
+
 	/** List of collection filter options */
 	TArray<TSharedPtr<FName>> CollectionsComboList;
 
