The 'Audit Asset' dialog will now correctly refresh it's display of Chunks for the 'Editor' platform.

AssetManager: Made sure that we clear CachedChunkMap before updating it so that ChunkIDs that have vbeen deleted during the editing session are properly discarded.
SAssetAuditBrowser: Instead of just refreshing the database for UAssetManager we refresh the data for the ManagerEditorModule, to make sure that it recaches the ChunkID's. This needs to be done once the dialog has been built as it can invoke callbacks to the SAssetAuditBrowser itself. This will result in the panel having up to date data when it is first displayed rather than requiring the user to press refresh immediately after opening it.

Note that the panel will not give sizes for chunks if the project is not set to generate chunked packfiles due to an optimization in CL 5672140, if that should be fixed or if we just prevent the chunks from being shown if the project will not create them is something for further consideration.

diff --git a/Engine/Plugins/Editor/AssetManagerEditor/Source/AssetManagerEditor/Private/SAssetAuditBrowser.cpp b/Engine/Plugins/Editor/AssetManagerEditor/Source/AssetManagerEditor/Private/SAssetAuditBrowser.cpp
--- a/Engine/Plugins/Editor/AssetManagerEditor/Source/AssetManagerEditor/Private/SAssetAuditBrowser.cpp
+++ b/Engine/Plugins/Editor/AssetManagerEditor/Source/AssetManagerEditor/Private/SAssetAuditBrowser.cpp
@@ -232,8 +232,6 @@ void SAssetAuditBrowser::Construct(const FArguments& InArgs)
 	AssetManager = &UAssetManager::Get();
 	EditorModule = &IAssetManagerEditorModule::Get();
 
-	AssetManager->UpdateManagementDatabase();
-
 	TArray<const FAssetManagerEditorRegistrySource*> AvailableSources;
 	ManagerEditorModule.GetAvailableRegistrySources(AvailableSources);
 
@@ -495,7 +493,9 @@ void SAssetAuditBrowser::Construct(const FArguments& InArgs)
 		]
 	];
 
-	RefreshAssetView();
+	// Refresh our data sources to make sure we are displaying up to date info otherwise the user will need to 
+	// hit "refresh" every time they open the dialog to be able to trust it
+	ManagerEditorModule.RefreshRegistryData();
 }
 
 FReply SAssetAuditBrowser::OnKeyDown(const FGeometry& MyGeometry, const FKeyEvent& InKeyEvent)
diff --git a/Engine/Source/Runtime/Engine/Private/AssetManager.cpp b/Engine/Source/Runtime/Engine/Private/AssetManager.cpp
--- a/Engine/Source/Runtime/Engine/Private/AssetManager.cpp
+++ b/Engine/Source/Runtime/Engine/Private/AssetManager.cpp
@@ -3032,7 +3032,10 @@ void UAssetManager::UpdateManagementDatabase(bool bForceRefresh)
 
 
 	TMultiMap<FAssetIdentifier, FAssetIdentifier> PrimaryAssetIdManagementMap;
-	TArray<int32> ChunkList, ExistingChunkList, OverrideChunkList;
+	TArray<int32> ChunkList;
+	TArray<int32> ExistingChunkList;
+
+	CachedChunkMap.Empty(); // Remove previous entries before we start adding to it
 
 	// Update management parent list, which is PrimaryAssetId -> PrimaryAssetId
 	for (const TPair<FName, TSharedRef<FPrimaryAssetTypeData>>& TypePair : AssetTypeMap)
@@ -3095,6 +3098,7 @@ void UAssetManager::UpdateManagementDatabase(bool bForceRefresh)
 	{
 		// Update the editor preview chunk package list for all chunks, but only if we actually care about chunks
 		// bGenerateChunks is settable per platform, but should be enabled on the default platform for preview to work
+		TArray<int32> OverrideChunkList;
 		for (FName PackageName : PackagesToUpdateChunksFor)
 		{
 			ChunkList.Reset();
