Tool Menus: Asset Manager Editor extensions convert to Tool Menus
#rnx

diff --git a/Engine/Plugins/Editor/AssetManagerEditor/Source/AssetManagerEditor/Private/AssetManagerEditorModule.cpp b/Engine/Plugins/Editor/AssetManagerEditor/Source/AssetManagerEditor/Private/AssetManagerEditorModule.cpp
--- a/Engine/Plugins/Editor/AssetManagerEditor/Source/AssetManagerEditor/Private/AssetManagerEditorModule.cpp
+++ b/Engine/Plugins/Editor/AssetManagerEditor/Source/AssetManagerEditor/Private/AssetManagerEditorModule.cpp
@@ -65,7 +65,9 @@
 #include "Editor/StatsViewer/Public/IStatsViewer.h"
 #include "Editor/StatsViewer/Public/StatsViewerModule.h"
 #include "Subsystems/AssetEditorSubsystem.h"
-
+#include "ToolMenus.h"
+#include "Toolkits/AssetEditorToolkitMenuContext.h"
+#include "ContentBrowserMenuContexts.h"
 
 #define LOCTEXT_NAMESPACE "AssetManagerEditor"
 
@@ -354,12 +356,9 @@ private:
 	static const FName ReferenceViewerTabName;
 	static const FName SizeMapTabName;
 
-	FDelegateHandle ContentBrowserAssetExtenderDelegateHandle;
-	FDelegateHandle ContentBrowserPathExtenderDelegateHandle;
 	FDelegateHandle ContentBrowserCommandExtenderDelegateHandle;
 	FDelegateHandle ReferenceViewerDelegateHandle;
 	FDelegateHandle AssetEditorExtenderDelegateHandle;
-	FDelegateHandle LevelEditorExtenderDelegateHandle;
 
 	TWeakPtr<SDockTab> AssetAuditTab;
 	TWeakPtr<SDockTab> ReferenceViewerTab;
@@ -376,13 +375,14 @@ private:
 	TSharedPtr<FAssetManagerGraphPanelNodeFactory> AssetManagerGraphPanelNodeFactory;
 	TSharedPtr<FAssetManagerGraphPanelPinFactory> AssetManagerGraphPanelPinFactory;
 
-	void CreateAssetContextMenu(FMenuBuilder& MenuBuilder);
+	static void CreateAssetContextMenu(FToolMenuSection& InSection);
 	void OnExtendContentBrowserCommands(TSharedRef<FUICommandList> CommandList, FOnContentBrowserGetSelection GetSelectionDelegate);
 	void OnExtendLevelEditorCommands(TSharedRef<FUICommandList> CommandList);
-	TSharedRef<FExtender> OnExtendContentBrowserAssetSelectionMenu(const TArray<FAssetData>& SelectedAssets);
-	TSharedRef<FExtender> OnExtendContentBrowserPathSelectionMenu(const TArray<FString>& SelectedPaths);
+	void ExtendContentBrowserAssetSelectionMenu();
+	void ExtendContentBrowserPathSelectionMenu();
 	TSharedRef<FExtender> OnExtendAssetEditor(const TSharedRef<FUICommandList> CommandList, const TArray<UObject*> ContextSensitiveObjects);
-	TSharedRef<FExtender> OnExtendLevelEditor(const TSharedRef<FUICommandList> CommandList, const TArray<AActor*> SelectedActors);
+	void ExtendAssetEditorMenu();
+	void ExtendLevelEditorActorContextMenu();
 	void OnHotReload(bool bWasTriggeredAutomatically);
 	void OnMarkPackageDirty(UPackage* Pkg, bool bWasDirty);
 	void OnEditAssetIdentifiers(TArray<FAssetIdentifier> AssetIdentifiers);
@@ -472,13 +472,9 @@ void FAssetManagerEditorModule::StartupModule()
 		// Register content browser hook
 		FContentBrowserModule& ContentBrowserModule = FModuleManager::LoadModuleChecked<FContentBrowserModule>(TEXT("ContentBrowser"));
 		
-		TArray<FContentBrowserMenuExtender_SelectedAssets>& CBAssetMenuExtenderDelegates = ContentBrowserModule.GetAllAssetViewContextMenuExtenders();
-		CBAssetMenuExtenderDelegates.Add(FContentBrowserMenuExtender_SelectedAssets::CreateRaw(this, &FAssetManagerEditorModule::OnExtendContentBrowserAssetSelectionMenu));
-		ContentBrowserAssetExtenderDelegateHandle = CBAssetMenuExtenderDelegates.Last().GetHandle();
+		ExtendContentBrowserAssetSelectionMenu();
 
-		TArray<FContentBrowserMenuExtender_SelectedPaths>& CBPathExtenderDelegates = ContentBrowserModule.GetAllPathViewContextMenuExtenders();
-		CBPathExtenderDelegates.Add(FContentBrowserMenuExtender_SelectedPaths::CreateRaw(this, &FAssetManagerEditorModule::OnExtendContentBrowserPathSelectionMenu));
-		ContentBrowserPathExtenderDelegateHandle = CBPathExtenderDelegates.Last().GetHandle();
+		ExtendContentBrowserPathSelectionMenu();
 
 		TArray<FContentBrowserCommandExtender>& CBCommandExtenderDelegates = ContentBrowserModule.GetAllContentBrowserCommandExtenders();
 		CBCommandExtenderDelegates.Add(FContentBrowserCommandExtender::CreateRaw(this, &FAssetManagerEditorModule::OnExtendContentBrowserCommands));
@@ -488,6 +484,7 @@ void FAssetManagerEditorModule::StartupModule()
 		TArray<FAssetEditorExtender>& AssetEditorMenuExtenderDelegates = FAssetEditorToolkit::GetSharedMenuExtensibilityManager()->GetExtenderDelegates();
 		AssetEditorMenuExtenderDelegates.Add(FAssetEditorExtender::CreateRaw(this, &FAssetManagerEditorModule::OnExtendAssetEditor));
 		AssetEditorExtenderDelegateHandle = AssetEditorMenuExtenderDelegates.Last().GetHandle();
+		ExtendAssetEditorMenu();
 
 		// Register level editor hooks and commands
 		FLevelEditorModule& LevelEditorModule = FModuleManager::Get().LoadModuleChecked<FLevelEditorModule>("LevelEditor");
@@ -495,9 +492,7 @@ void FAssetManagerEditorModule::StartupModule()
 		TSharedRef<FUICommandList> CommandList = LevelEditorModule.GetGlobalLevelEditorActions();
 		OnExtendLevelEditorCommands(CommandList);
 
-		TArray<FLevelEditorModule::FLevelViewportMenuExtender_SelectedActors>& LevelEditorMenuExtenderDelegates = LevelEditorModule.GetAllLevelViewportContextMenuExtenders();
-		LevelEditorMenuExtenderDelegates.Add(FLevelEditorModule::FLevelViewportMenuExtender_SelectedActors::CreateRaw(this, &FAssetManagerEditorModule::OnExtendLevelEditor));
-		LevelEditorExtenderDelegateHandle = AssetEditorMenuExtenderDelegates.Last().GetHandle();
+		ExtendLevelEditorActorContextMenu();
 
 		// Add nomad tabs
 		FGlobalTabmanager::Get()->RegisterNomadTabSpawner(AssetAuditTabName, FOnSpawnTab::CreateRaw(this, &FAssetManagerEditorModule::SpawnAssetAuditTab))
@@ -551,10 +546,6 @@ void FAssetManagerEditorModule::ShutdownModule()
 	{
 		FContentBrowserModule& ContentBrowserModule = FModuleManager::LoadModuleChecked<FContentBrowserModule>(TEXT("ContentBrowser"));
 		
-		TArray<FContentBrowserMenuExtender_SelectedAssets>& CBAssetMenuExtenderDelegates = ContentBrowserModule.GetAllAssetViewContextMenuExtenders();
-		CBAssetMenuExtenderDelegates.RemoveAll([this](const FContentBrowserMenuExtender_SelectedAssets& Delegate) { return Delegate.GetHandle() == ContentBrowserAssetExtenderDelegateHandle; });
-		TArray<FContentBrowserMenuExtender_SelectedPaths>& CBPathMenuExtenderDelegates = ContentBrowserModule.GetAllPathViewContextMenuExtenders();
-		CBPathMenuExtenderDelegates.RemoveAll([this](const FContentBrowserMenuExtender_SelectedPaths& Delegate) { return Delegate.GetHandle() == ContentBrowserPathExtenderDelegateHandle; });
 		TArray<FContentBrowserCommandExtender>& CBCommandExtenderDelegates = ContentBrowserModule.GetAllContentBrowserCommandExtenders();
 		CBCommandExtenderDelegates.RemoveAll([this](const FContentBrowserCommandExtender& Delegate) { return Delegate.GetHandle() == ContentBrowserCommandExtenderDelegateHandle; });
 		
@@ -574,9 +565,6 @@ void FAssetManagerEditorModule::ShutdownModule()
 		CommandList->UnmapAction(FAssetManagerEditorCommands::Get().ViewAssetAudit);
 		CommandList->UnmapAction(FAssetManagerEditorCommands::Get().ViewShaderCookStatistics);
 
-		TArray<FLevelEditorModule::FLevelViewportMenuExtender_SelectedActors>& LevelEditorMenuExtenderDelegates = LevelEditorModule.GetAllLevelViewportContextMenuExtenders();
-		LevelEditorMenuExtenderDelegates.RemoveAll([this](const FLevelEditorModule::FLevelViewportMenuExtender_SelectedActors& Delegate) { return Delegate.GetHandle() == LevelEditorExtenderDelegateHandle; });
-
 		if (AssetManagerGraphPanelNodeFactory.IsValid())
 		{
 			FEdGraphUtilities::UnregisterVisualNodeFactory(AssetManagerGraphPanelNodeFactory);
@@ -821,12 +809,12 @@ TArray<FName> FAssetManagerEditorModule::GetContentBrowserSelectedAssetPackages(
 	return OutAssetPackages;
 }
 
-void FAssetManagerEditorModule::CreateAssetContextMenu(FMenuBuilder& MenuBuilder)
+void FAssetManagerEditorModule::CreateAssetContextMenu(FToolMenuSection& InSection)
 {
-	MenuBuilder.AddMenuEntry(FAssetManagerEditorCommands::Get().ViewReferences);
-	MenuBuilder.AddMenuEntry(FAssetManagerEditorCommands::Get().ViewSizeMap);
-	MenuBuilder.AddMenuEntry(FAssetManagerEditorCommands::Get().ViewAssetAudit);
-	MenuBuilder.AddMenuEntry(FAssetManagerEditorCommands::Get().ViewShaderCookStatistics);
+	InSection.AddMenuEntry(FAssetManagerEditorCommands::Get().ViewReferences);
+	InSection.AddMenuEntry(FAssetManagerEditorCommands::Get().ViewSizeMap);
+	InSection.AddMenuEntry(FAssetManagerEditorCommands::Get().ViewAssetAudit);
+	InSection.AddMenuEntry(FAssetManagerEditorCommands::Get().ViewShaderCookStatistics);
 }
 
 void FAssetManagerEditorModule::OnExtendContentBrowserCommands(TSharedRef<FUICommandList> CommandList, FOnContentBrowserGetSelection GetSelectionDelegate)
@@ -899,30 +887,32 @@ void FAssetManagerEditorModule::OnExtendLevelEditorCommands(TSharedRef<FUIComman
 	);
 }
 
-TSharedRef<FExtender> FAssetManagerEditorModule::OnExtendContentBrowserAssetSelectionMenu(const TArray<FAssetData>& SelectedAssets)
+void FAssetManagerEditorModule::ExtendContentBrowserAssetSelectionMenu()
 {
-	TSharedRef<FExtender> Extender(new FExtender());
-	
-	Extender->AddMenuExtension(
-		"AssetContextReferences",
-		EExtensionHook::After,
-		nullptr,
-		FMenuExtensionDelegate::CreateRaw(this, &FAssetManagerEditorModule::CreateAssetContextMenu));
-
-	return Extender;
+	UToolMenu* Menu = UToolMenus::Get()->ExtendMenu("ContentBrowser.AssetContextMenu");
+	FToolMenuSection& Section = Menu->FindOrAddSection("AssetContextReferences");
+	FToolMenuEntry& Entry = Section.AddDynamicEntry("AssetManagerEditorViewCommands", FNewToolMenuSectionDelegate::CreateLambda([](FToolMenuSection& InSection)
+	{
+		UContentBrowserAssetContextMenuContext* Context = InSection.FindContext<UContentBrowserAssetContextMenuContext>();
+		if (Context && Context->bCanBeModified)
+		{
+			FAssetManagerEditorModule::CreateAssetContextMenu(InSection);
+		}
+	}));
 }
 
-TSharedRef<FExtender> FAssetManagerEditorModule::OnExtendContentBrowserPathSelectionMenu(const TArray<FString>& SelectedPaths)
+void FAssetManagerEditorModule::ExtendContentBrowserPathSelectionMenu()
 {
-	TSharedRef<FExtender> Extender(new FExtender());
-
-	Extender->AddMenuExtension(
-		"PathContextBulkOperations",
-		EExtensionHook::After,
-		nullptr,
-		FMenuExtensionDelegate::CreateRaw(this, &FAssetManagerEditorModule::CreateAssetContextMenu));
-
-	return Extender;
+	UToolMenu* Menu = UToolMenus::Get()->ExtendMenu("ContentBrowser.FolderContextMenu");
+	FToolMenuSection& Section = Menu->FindOrAddSection("PathContextBulkOperations");
+	FToolMenuEntry& Entry = Section.AddDynamicEntry("AssetManagerEditorViewCommands", FNewToolMenuSectionDelegate::CreateLambda([](FToolMenuSection& InSection)
+	{
+		UContentBrowserFolderContext* Context = InSection.FindContext<UContentBrowserFolderContext>();
+		if (Context && Context->bCanBeModified && Context->NumAssetPaths > 0)
+		{
+			FAssetManagerEditorModule::CreateAssetContextMenu(InSection);
+		}
+	}));
 }
 
 TSharedRef<FExtender> FAssetManagerEditorModule::OnExtendAssetEditor(const TSharedRef<FUICommandList> CommandList, const TArray<UObject*> ContextSensitiveObjects)
@@ -956,32 +946,44 @@ TSharedRef<FExtender> FAssetManagerEditorModule::OnExtendAssetEditor(const TShar
 		CommandList->MapAction(
 			FAssetManagerEditorCommands::Get().ViewAssetAudit,
 			FExecuteAction::CreateRaw(this, &FAssetManagerEditorModule::OpenAssetAuditUI, PackageNames));
-
-		Extender->AddMenuExtension(
-			"FindInContentBrowser",
-			EExtensionHook::After,
-			CommandList,
-			FMenuExtensionDelegate::CreateRaw(this, &FAssetManagerEditorModule::CreateAssetContextMenu));
 	}
 
 	return Extender;
 }
 
-TSharedRef<FExtender> FAssetManagerEditorModule::OnExtendLevelEditor(const TSharedRef<FUICommandList> CommandList, const TArray<AActor*> SelectedActors)
+void FAssetManagerEditorModule::ExtendAssetEditorMenu()
 {
-	TSharedRef<FExtender> Extender(new FExtender());
-
-	// Most of this work is handled in the command list extension
-	if (SelectedActors.Num() > 0)
+	UToolMenu* Menu = UToolMenus::Get()->ExtendMenu("MainFrame.MainMenu.Asset");
+	FToolMenuSection& Section = Menu->FindOrAddSection("AssetEditorActions");
+	FToolMenuEntry& Entry = Section.AddDynamicEntry("AssetManagerEditorViewCommands", FNewToolMenuSectionDelegate::CreateLambda([](FToolMenuSection& InSection)
 	{
-		Extender->AddMenuExtension(
-			"ActorAsset",
-			EExtensionHook::After,
-			CommandList,
-			FMenuExtensionDelegate::CreateRaw(this, &FAssetManagerEditorModule::CreateAssetContextMenu));
-	}
+		UAssetEditorToolkitMenuContext* MenuContext = InSection.FindContext<UAssetEditorToolkitMenuContext>();
+		if (MenuContext && MenuContext->Toolkit.IsValid() && MenuContext->Toolkit.Pin()->IsActuallyAnAsset())
+		{
+			for (const UObject* EditedAsset : *MenuContext->Toolkit.Pin()->GetObjectsCurrentlyBeingEdited())
+			{
+				if (EditedAsset && EditedAsset->IsAsset() && !EditedAsset->IsPendingKill())
+				{
+					FAssetManagerEditorModule::CreateAssetContextMenu(InSection);
+					break;
+				}
+			}
+		}
+	}));
+	Entry.InsertPosition = FToolMenuInsert("FindInContentBrowser", EToolMenuInsertType::After);
+}
 
-	return Extender;
+void FAssetManagerEditorModule::ExtendLevelEditorActorContextMenu()
+{
+	UToolMenu* Menu = UToolMenus::Get()->ExtendMenu("LevelEditor.ActorContextMenu");
+	FToolMenuSection& Section = Menu->FindOrAddSection("ActorAsset");
+	FToolMenuEntry& Entry = Section.AddDynamicEntry("AssetManagerEditorViewCommands", FNewToolMenuSectionDelegate::CreateLambda([](FToolMenuSection& InSection)
+	{
+		if (GEditor->CanSyncToContentBrowser() && GEditor->GetSelectedComponentCount() == 0 && GEditor->GetSelectedActorCount() > 0)
+		{
+			FAssetManagerEditorModule::CreateAssetContextMenu(InSection);
+		}
+	}));
 }
 
 void FAssetManagerEditorModule::OnHotReload(bool bWasTriggeredAutomatically)
diff --git a/Engine/Source/Editor/ContentBrowser/Private/AssetContextMenu.cpp b/Engine/Source/Editor/ContentBrowser/Private/AssetContextMenu.cpp
--- a/Engine/Source/Editor/ContentBrowser/Private/AssetContextMenu.cpp
+++ b/Engine/Source/Editor/ContentBrowser/Private/AssetContextMenu.cpp
@@ -1332,9 +1332,7 @@ bool FAssetContextMenu::AddReferenceMenuOptions(UToolMenu* Menu)
 	UContentBrowserAssetContextMenuContext* Context = Menu->FindContext<UContentBrowserAssetContextMenuContext>();
 
 	{
-		// If can not be modified change section name to prevent reference menu extensions
-		const FName SectionName = Context->bCanBeModified ? FName("AssetContextReferences") : FName("AssetContextReferencesReadOnly");
-		FToolMenuSection& Section = Menu->AddSection(SectionName, LOCTEXT("ReferencesMenuHeading", "References"));
+		FToolMenuSection& Section = Menu->AddSection("AssetContextReferences", LOCTEXT("ReferencesMenuHeading", "References"));
 
 		Section.AddMenuEntry(
 			"CopyReference",
diff --git a/Engine/Source/Editor/ContentBrowser/Private/PathContextMenu.cpp b/Engine/Source/Editor/ContentBrowser/Private/PathContextMenu.cpp
--- a/Engine/Source/Editor/ContentBrowser/Private/PathContextMenu.cpp
+++ b/Engine/Source/Editor/ContentBrowser/Private/PathContextMenu.cpp
@@ -123,8 +123,8 @@ void FPathContextMenu::MakePathViewContextMenu(UToolMenu* Menu)
 	// Useful for actions whose "CanExecute" will not change or is expensive to calculate.
 	CacheCanExecuteVars();
 
-	int32 NumAssetPaths, NumClassPaths;
-	ContentBrowserUtils::CountPathTypes(SelectedPaths, NumAssetPaths, NumClassPaths);
+	const int32 NumAssetPaths = Context->NumAssetPaths;
+	const int32 NumClassPaths = Context->NumClassPaths;
 
 	// Only add something if at least one folder is selected
 	if ( SelectedPaths.Num() > 0 )
diff --git a/Engine/Source/Editor/ContentBrowser/Private/SContentBrowser.cpp b/Engine/Source/Editor/ContentBrowser/Private/SContentBrowser.cpp
--- a/Engine/Source/Editor/ContentBrowser/Private/SContentBrowser.cpp
+++ b/Engine/Source/Editor/ContentBrowser/Private/SContentBrowser.cpp
@@ -3218,6 +3218,7 @@ TSharedPtr<SWidget> SContentBrowser::GetFolderContextMenu(const TArray<FString>&
 	UContentBrowserFolderContext* Context = NewObject<UContentBrowserFolderContext>();
 	Context->ContentBrowser = SharedThis(this);
 	Context->OnCreateNewFolder = InOnCreateNewFolder;
+	ContentBrowserUtils::CountPathTypes(SelectedPaths, Context->NumAssetPaths, Context->NumClassPaths);
 
 	FAssetToolsModule& AssetToolsModule = FModuleManager::LoadModuleChecked<FAssetToolsModule>("AssetTools");
 	Context->bCanBeModified = AssetToolsModule.Get().AllPassWritableFolderFilter(SelectedPaths);
diff --git a/Engine/Source/Editor/ContentBrowser/Public/ContentBrowserMenuContexts.h b/Engine/Source/Editor/ContentBrowser/Public/ContentBrowserMenuContexts.h
--- a/Engine/Source/Editor/ContentBrowser/Public/ContentBrowserMenuContexts.h
+++ b/Engine/Source/Editor/ContentBrowser/Public/ContentBrowserMenuContexts.h
@@ -67,6 +67,12 @@ public:
 	UPROPERTY()
 	bool bCanBeModified;
 
+	UPROPERTY()
+	int32 NumAssetPaths;
+
+	UPROPERTY()
+	int32 NumClassPaths;
+
 	TWeakPtr<SContentBrowser> ContentBrowser;
 	FOnCreateNewFolder OnCreateNewFolder;
 };
