- Fix editor typo "relevent" to "relevant"

diff --git a/Engine/Plugins/Editor/AssetManagerEditor/Source/AssetManagerEditor/Private/AssetManagerEditorCommands.cpp b/Engine/Plugins/Editor/AssetManagerEditor/Source/AssetManagerEditor/Private/AssetManagerEditorCommands.cpp
--- a/Engine/Plugins/Editor/AssetManagerEditor/Source/AssetManagerEditor/Private/AssetManagerEditorCommands.cpp
+++ b/Engine/Plugins/Editor/AssetManagerEditor/Source/AssetManagerEditor/Private/AssetManagerEditorCommands.cpp
@@ -23,7 +23,7 @@ void FAssetManagerEditorCommands::RegisterCommands()
 	
 	UI_COMMAND(ViewAssetAudit, "Audit Assets...", "Opens the Asset Audit UI and displays information about the selected assets", EUserInterfaceActionType::Button, FInputChord(EModifierKey::Shift | EModifierKey::Alt, EKeys::A));
 
-	UI_COMMAND(OpenSelectedInAssetEditor, "Edit...", "Opens the selected asset in the relevent editor.", EUserInterfaceActionType::Button, FInputChord());
+	UI_COMMAND(OpenSelectedInAssetEditor, "Edit...", "Opens the selected asset in the relevant editor.", EUserInterfaceActionType::Button, FInputChord());
 	UI_COMMAND(ZoomToFit, "Zoom to Fit", "Zoom in and center the view on the selected item", EUserInterfaceActionType::Button, FInputChord(EKeys::Home));
 	UI_COMMAND(ReCenterGraph, "Re-Center Graph", "Re-centers the graph on this node, showing all referencers and references for this asset instead", EUserInterfaceActionType::Button, FInputChord());
 	UI_COMMAND(CopyReferencedObjects, "Copy Referenced Objects List", "Copies the list of objects that the selected asset references to the clipboard.", EUserInterfaceActionType::Button, FInputChord());
