Display actor labels in Reference Viewer

diff --git a/Engine/Plugins/Editor/AssetManagerEditor/Source/AssetManagerEditor/Private/ReferenceViewer/EdGraphNode_Reference.cpp b/Engine/Plugins/Editor/AssetManagerEditor/Source/AssetManagerEditor/Private/ReferenceViewer/EdGraphNode_Reference.cpp
--- a/Engine/Plugins/Editor/AssetManagerEditor/Source/AssetManagerEditor/Private/ReferenceViewer/EdGraphNode_Reference.cpp
+++ b/Engine/Plugins/Editor/AssetManagerEditor/Source/AssetManagerEditor/Private/ReferenceViewer/EdGraphNode_Reference.cpp
@@ -57,7 +57,11 @@ void UEdGraphNode_Reference::SetupReferenceNode(const FIntPoint& NodeLoc, const
 		{
 			NodeComment = First.PackageName.ToString();
 		}
+
+		static const FName NAME_ActorLabel(TEXT("ActorLabel"));
+
 		NodeTitle = FText::FromString(MainAssetName);
+		InAssetData.GetTagValue(NAME_ActorLabel, NodeTitle);
 	}
 	else
 	{
