Update to use the correct style for comment bubbles

diff --git a/Engine/Plugins/Editor/AssetManagerEditor/Source/AssetManagerEditor/Private/ReferenceViewer/SReferenceNode.cpp b/Engine/Plugins/Editor/AssetManagerEditor/Source/AssetManagerEditor/Private/ReferenceViewer/SReferenceNode.cpp
--- a/Engine/Plugins/Editor/AssetManagerEditor/Source/AssetManagerEditor/Private/ReferenceViewer/SReferenceNode.cpp
+++ b/Engine/Plugins/Editor/AssetManagerEditor/Source/AssetManagerEditor/Private/ReferenceViewer/SReferenceNode.cpp
@@ -212,10 +212,9 @@ void SReferenceNode::UpdateGraphNode()
 	{
 		TSharedPtr<SCommentBubble> CommentBubble;
 
-		SAssignNew( CommentBubble, SCommentBubble )
-		.GraphNode( GraphNode )
-		.Text( this, &SGraphNode::GetNodeComment )
-		.ColorAndOpacity( this, &SReferenceNode::GetCommentColor );
+		SAssignNew(CommentBubble, SCommentBubble)
+			.GraphNode(GraphNode)
+			.Text(this, &SGraphNode::GetNodeComment);
 
 		GetOrAddSlot( ENodeZone::TopCenter )
 		.SlotOffset( TAttribute<FVector2D>( CommentBubble.Get(), &SCommentBubble::GetOffset ))
diff --git a/Engine/Source/Editor/GraphEditor/Private/SCommentBubble.cpp b/Engine/Source/Editor/GraphEditor/Private/SCommentBubble.cpp
--- a/Engine/Source/Editor/GraphEditor/Private/SCommentBubble.cpp
+++ b/Engine/Source/Editor/GraphEditor/Private/SCommentBubble.cpp
@@ -251,7 +251,7 @@ void SCommentBubble::UpdateBubble()
 						[
 							SAssignNew(TextBlock, SMultiLineEditableTextBox)
 							.Text(MakeAttributeLambda([this] { return CachedCommentText; }))
-							.Style(FAppStyle::Get(), "Graph.Node.NodeTitleEditableText")
+							.Style(FAppStyle::Get(), "Graph.Node.NodeTitleInlineEditableText")
 							.HintText( NSLOCTEXT( "CommentBubble", "EditCommentHint", "Click to edit" ))
 							.IsReadOnly(this, &SCommentBubble::IsReadOnly)
 							.SelectAllTextWhenFocused( true )
