DevelopmentAssetRegistry.bin: Reduce duplication by making a function to return this string.
#rnx

diff --git a/Engine/Plugins/Editor/AssetManagerEditor/Source/AssetManagerEditor/Private/AssetManagerEditorModule.cpp b/Engine/Plugins/Editor/AssetManagerEditor/Source/AssetManagerEditor/Private/AssetManagerEditorModule.cpp
--- a/Engine/Plugins/Editor/AssetManagerEditor/Source/AssetManagerEditor/Private/AssetManagerEditorModule.cpp
+++ b/Engine/Plugins/Editor/AssetManagerEditor/Source/AssetManagerEditor/Private/AssetManagerEditorModule.cpp
@@ -2,7 +2,7 @@
 
 #include "AssetManagerEditorModule.h"
 #include "Modules/ModuleManager.h"
-#include "AssetRegistryModule.h"
+#include "AssetRegistry/AssetRegistryModule.h"
 #include "PrimaryAssetTypeCustomization.h"
 #include "PrimaryAssetIdCustomization.h"
 #include "SAssetAuditBrowser.h"
@@ -1405,11 +1405,11 @@ FString FAssetManagerEditorModule::GetSavedAssetRegistryPath(ITargetPlatform* Ta
 	CommandLinePath.ReplaceInline(TEXT("[Platform]"), *PlatformName);
 	
 	// We can only load DevelopmentAssetRegistry, the normal asset registry doesn't have enough data to be useful
-	FString CookedDevelopmentAssetRegistry = FPaths::ProjectDir() / TEXT("Metadata") / TEXT("DevelopmentAssetRegistry.bin");
+	FString CookedDevelopmentAssetRegistry = FPaths::ProjectDir() / TEXT("Metadata") / GetDevelopmentAssetRegistryFilename();
 
 	FString DevCookedPath = CookedSandbox->ConvertToAbsolutePathForExternalAppForWrite(*CookedDevelopmentAssetRegistry).Replace(TEXT("[Platform]"), *PlatformName);
 	FString DevEditorCookedPath = EditorCookedSandbox->ConvertToAbsolutePathForExternalAppForWrite(*CookedDevelopmentAssetRegistry).Replace(TEXT("[Platform]"), *PlatformName);
-	FString DevSharedCookedPath = FPaths::Combine(*FPaths::ProjectSavedDir(), TEXT("SharedIterativeBuild"), PlatformName, TEXT("Metadata"), TEXT("DevelopmentAssetRegistry.bin"));
+	FString DevSharedCookedPath = FPaths::Combine(*FPaths::ProjectSavedDir(), TEXT("SharedIterativeBuild"), PlatformName, TEXT("Metadata"), GetDevelopmentAssetRegistryFilename());
 
 	// Try command line, then cooked, then shared build
 	if (!CommandLinePath.IsEmpty() && IFileManager::Get().FileExists(*CommandLinePath))
@@ -1474,15 +1474,16 @@ void FAssetManagerEditorModule::SetCurrentRegistrySource(const FString& SourceNa
 
 			IDesktopPlatform* DesktopPlatform = FDesktopPlatformModule::Get();
 			const void* ParentWindowWindowHandle = FSlateApplication::Get().FindBestParentWindowHandleForDialogs(nullptr);
-			const FText Title = LOCTEXT("LoadAssetRegistry", "Load DevelopmentAssetRegistry.bin");
-			const FString FileTypes = TEXT("DevelopmentAssetRegistry.bin|*.bin");
+			const TCHAR* DevelopmentAssetRegistryFilename = GetDevelopmentAssetRegistryFilename();
+			const FText Title = LOCTEXT("LoadAssetRegistry", "Load DevelopmentAssetRegistry");
+			const FString FileTypes = FString::Printf(TEXT("%s|*.bin"), DevelopmentAssetRegistryFilename);
 
 			TArray<FString> OutFilenames;
 			DesktopPlatform->OpenFileDialog(
 				ParentWindowWindowHandle,
 				Title.ToString(),
 				TEXT(""),
-				TEXT("DevelopmentAssetRegistry.bin"),
+				DevelopmentAssetRegistryFilename,
 				FileTypes,
 				EFileDialogFlags::None,
 				OutFilenames
diff --git a/Engine/Source/Developer/IoStoreUtilities/Private/IoStoreUtilities.cpp b/Engine/Source/Developer/IoStoreUtilities/Private/IoStoreUtilities.cpp
--- a/Engine/Source/Developer/IoStoreUtilities/Private/IoStoreUtilities.cpp
+++ b/Engine/Source/Developer/IoStoreUtilities/Private/IoStoreUtilities.cpp
@@ -45,7 +45,7 @@
 #include "Async/Async.h"
 #include "RSA.h"
 #include "Misc/AssetRegistryInterface.h"
-#include "AssetRegistryState.h"
+#include "AssetRegistry/AssetRegistryState.h"
 #include "Misc/OutputDeviceFile.h"
 #include "Misc/FeedbackContext.h"
 #include "Serialization/LargeMemoryReader.h"
@@ -4712,7 +4712,7 @@ int32 CreateIoStoreContainerFiles(const TCHAR* CmdLine)
 			return -1;
 		}
 
-		FString DevelopmentAssetRegistryPath = FPaths::Combine(Arguments.BasedOnReleaseVersionPath, TEXT("Metadata/DevelopmentAssetRegistry.bin"));
+		FString DevelopmentAssetRegistryPath = FPaths::Combine(Arguments.BasedOnReleaseVersionPath, TEXT("Metadata"), GetDevelopmentAssetRegistryFilename());
 
 		bool bAssetRegistryLoaded = false;
 		FArrayReader SerializedAssetData;
diff --git a/Engine/Source/Editor/UnrealEd/Private/Commandlets/AssetRegistryGenerator.cpp b/Engine/Source/Editor/UnrealEd/Private/Commandlets/AssetRegistryGenerator.cpp
--- a/Engine/Source/Editor/UnrealEd/Private/Commandlets/AssetRegistryGenerator.cpp
+++ b/Engine/Source/Editor/UnrealEd/Private/Commandlets/AssetRegistryGenerator.cpp
@@ -18,7 +18,7 @@
 #include "ICollectionManager.h"
 #include "CollectionManagerModule.h"
 #include "Interfaces/ITargetPlatform.h"
-#include "AssetRegistryModule.h"
+#include "AssetRegistry/AssetRegistryModule.h"
 #include "GameDelegates.h"
 #include "Commandlets/IChunkDataGenerator.h"
 #include "Commandlets/ChunkDependencyInfo.h"
@@ -1402,14 +1402,15 @@ bool FAssetRegistryGenerator::SaveAssetRegistry(const FString& SandboxPath, bool
 
 		// Save the generated registry
 		FString PlatformSandboxPath = SandboxPath.Replace(TEXT("[Platform]"), *TargetPlatform->PlatformName());
-		PlatformSandboxPath.ReplaceInline(TEXT("AssetRegistry.bin"), TEXT("Metadata/DevelopmentAssetRegistry.bin"));
+		const TCHAR* DevelopmentAssetRegistryFilename = GetDevelopmentAssetRegistryFilename();
+		PlatformSandboxPath.ReplaceInline(TEXT("AssetRegistry.bin"), *FString::Printf(TEXT("Metadata/%s"), DevelopmentAssetRegistryFilename));
 		FFileHelper::SaveArrayToFile(SerializedAssetRegistry, *PlatformSandboxPath);
 
 		UE_LOG(LogAssetRegistryGenerator, Display, TEXT("Generated development asset registry %s num assets %d, size is %5.2fkb"), *PlatformSandboxPath, State.GetNumAssets(), (float)SerializedAssetRegistry.Num() / 1024.f);
 
 		if (bGenerateChunks && bUseAssetManager)
 		{
-			FString ChunkListsPath = PlatformSandboxPath.Replace(TEXT("/DevelopmentAssetRegistry.bin"), TEXT(""));
+			FString ChunkListsPath = PlatformSandboxPath.Replace(*FString::Printf(TEXT("/%s"), DevelopmentAssetRegistryFilename), TEXT(""));
 
 			// Write out CSV file with chunking information
 			GenerateAssetChunkInformationCSV(ChunkListsPath, false);
diff --git a/Engine/Source/Editor/UnrealEd/Private/Commandlets/DiffAssetRegistriesCommandlet.cpp b/Engine/Source/Editor/UnrealEd/Private/Commandlets/DiffAssetRegistriesCommandlet.cpp
--- a/Engine/Source/Editor/UnrealEd/Private/Commandlets/DiffAssetRegistriesCommandlet.cpp
+++ b/Engine/Source/Editor/UnrealEd/Private/Commandlets/DiffAssetRegistriesCommandlet.cpp
@@ -2,7 +2,8 @@
 
 #include "Commandlets/DiffAssetRegistriesCommandlet.h"
 
-#include "AssetData.h"
+#include "AssetRegistry/AssetData.h"
+#include "AssetRegistry/IAssetRegistry.h"
 
 #include "UObject/Class.h"
 #include "PlatformInfo.h"
@@ -185,7 +186,6 @@ int32 UDiffAssetRegistriesCommandlet::Main(const FString& FullCommandLine)
 			return false;
 		};
 
-	// const TCHAR* AssetRegistrySubPath = TEXT("/Metadata/DevelopmentAssetRegistry.bin");
 	const FString* OldPathVal = Params.Find(FString(TEXT("OldPath")));
 	if (OldPathVal)
 	{
@@ -253,7 +253,8 @@ int32 UDiffAssetRegistriesCommandlet::Main(const FString& FullCommandLine)
 	FPaths::NormalizeFilename(OldPath);
 
 	// try to discern platform
-	/*if (NewPath.Contains(AssetRegistrySubPath))
+	/*FString AssetRegistrySubPath = FString::Printf(TEXT("/Metadata/%s"), GetDevelopmentAssetRegistryFilename());
+	if (NewPath.Contains(AssetRegistrySubPath))
 	{
 		FString NewPlatformDir = NewPath.Left(NewPath.Find(AssetRegistrySubPath));
 		FString PlatformPath = FPaths::GetCleanFilename(NewPlatformDir);
diff --git a/Engine/Source/Editor/UnrealEd/Private/CookOnTheFlyServer.cpp b/Engine/Source/Editor/UnrealEd/Private/CookOnTheFlyServer.cpp
--- a/Engine/Source/Editor/UnrealEd/Private/CookOnTheFlyServer.cpp
+++ b/Engine/Source/Editor/UnrealEd/Private/CookOnTheFlyServer.cpp
@@ -9,8 +9,8 @@
 #include "Algo/AnyOf.h"
 #include "Algo/Find.h"
 #include "AssetCompilingManager.h"
-#include "AssetRegistryModule.h"
-#include "AssetRegistryState.h"
+#include "AssetRegistry/AssetRegistryModule.h"
+#include "AssetRegistry/AssetRegistryState.h"
 #include "Async/Async.h"
 #include "Async/ParallelFor.h"
 #include "Commandlets/AssetRegistryGenerator.h"
@@ -232,12 +232,6 @@ const FString& GetAssetRegistryFilename()
 	return AssetRegistryFilename;
 }
 
-const FString& GetDevelopmentAssetRegistryFilename()
-{
-	static const FString DevelopmentAssetRegistryFilename = FString(TEXT("DevelopmentAssetRegistry.bin"));
-	return DevelopmentAssetRegistryFilename;
-}
-
 /**
  * Uses the FMessageLog to log a message
  * 
@@ -7015,8 +7009,9 @@ void UCookOnTheFlyServer::CookByTheBookFinished()
 					IFileManager::Get().Copy(*VersionedRegistryFilename, *CookedAssetRegistryFilename, true, true);
 
 					// Also copy development registry if it exists
-					const FString DevVersionedRegistryFilename = VersionedRegistryFilename.Replace(TEXT("AssetRegistry.bin"), TEXT("Metadata/DevelopmentAssetRegistry.bin"));
-					const FString DevCookedAssetRegistryFilename = CookedAssetRegistryFilename.Replace(TEXT("AssetRegistry.bin"), TEXT("Metadata/DevelopmentAssetRegistry.bin"));
+					FString DevelopmentAssetRegistryRelativePath = FString::Printf(TEXT("Metadata/%s"), GetDevelopmentAssetRegistryFilename());
+					const FString DevVersionedRegistryFilename = VersionedRegistryFilename.Replace(TEXT("AssetRegistry.bin"), *DevelopmentAssetRegistryRelativePath);
+					const FString DevCookedAssetRegistryFilename = CookedAssetRegistryFilename.Replace(TEXT("AssetRegistry.bin"), *DevelopmentAssetRegistryRelativePath);
 					IFileManager::Get().Copy(*DevVersionedRegistryFilename, *DevCookedAssetRegistryFilename, true, true);
 				}
 			}
diff --git a/Engine/Source/Editor/UnrealEd/Private/Cooker/CookTypes.h b/Engine/Source/Editor/UnrealEd/Private/Cooker/CookTypes.h
--- a/Engine/Source/Editor/UnrealEd/Private/Cooker/CookTypes.h
+++ b/Engine/Source/Editor/UnrealEd/Private/Cooker/CookTypes.h
@@ -260,6 +260,5 @@ public:
 };
 
 void LogCookerMessage(const FString& MessageText, EMessageSeverity::Type Severity);
-const FString& GetDevelopmentAssetRegistryFilename();
 constexpr uint32 ExpectedMaxNumPlatforms = 32;
 #define REMAPPED_PLUGINS TEXT("RemappedPlugins")
diff --git a/Engine/Source/Editor/UnrealEd/Private/Cooker/LooseCookedPackageWriter.cpp b/Engine/Source/Editor/UnrealEd/Private/Cooker/LooseCookedPackageWriter.cpp
--- a/Engine/Source/Editor/UnrealEd/Private/Cooker/LooseCookedPackageWriter.cpp
+++ b/Engine/Source/Editor/UnrealEd/Private/Cooker/LooseCookedPackageWriter.cpp
@@ -76,7 +76,7 @@ void FLooseCookedPackageWriter::WriteLinkerAdditionalData(const FLinkerAdditiona
 
 FDateTime FLooseCookedPackageWriter::GetPreviousCookTime() const
 {
-	const FString PreviousAssetRegistry = FPaths::Combine(MetadataDirectoryPath, TEXT("DevelopmentAssetRegistry.bin"));
+	const FString PreviousAssetRegistry = FPaths::Combine(MetadataDirectoryPath, GetDevelopmentAssetRegistryFilename());
 	return IFileManager::Get().GetTimeStamp(*PreviousAssetRegistry);
 }
 
diff --git a/Engine/Source/Runtime/AssetRegistry/Private/AssetRegistry.cpp b/Engine/Source/Runtime/AssetRegistry/Private/AssetRegistry.cpp
--- a/Engine/Source/Runtime/AssetRegistry/Private/AssetRegistry.cpp
+++ b/Engine/Source/Runtime/AssetRegistry/Private/AssetRegistry.cpp
@@ -291,6 +291,11 @@ EChunkProgressReportingType::Type GetChunkAvailabilityProgressType(EAssetAvailab
 	return ChunkReportType;
 }
 
+const TCHAR* GetDevelopmentAssetRegistryFilename()
+{
+	return TEXT("DevelopmentAssetRegistry.bin");
+}
+
 UAssetRegistry::UAssetRegistry(const FObjectInitializer& ObjectInitializer)
 	: Super(ObjectInitializer)
 {
diff --git a/Engine/Source/Runtime/AssetRegistry/Public/AssetRegistry/IAssetRegistry.h b/Engine/Source/Runtime/AssetRegistry/Public/AssetRegistry/IAssetRegistry.h
--- a/Engine/Source/Runtime/AssetRegistry/Public/AssetRegistry/IAssetRegistry.h
+++ b/Engine/Source/Runtime/AssetRegistry/Public/AssetRegistry/IAssetRegistry.h
@@ -642,3 +642,7 @@ namespace AssetRegistry
 	ASSETREGISTRY_API bool ReadPackageDataDependencies(FArchive& BinaryArchive, TBitArray<>& OutImportUsedInGame, TBitArray<>& OutSoftPackageUsedInGame);
 }
 }
+
+/** Returns the filename without filepath for the DevelopmentAssetRegistry written by the cooker. */
+ASSETREGISTRY_API const TCHAR* GetDevelopmentAssetRegistryFilename();
+
