#UE Fix selected path processing in AssetManagerEditorModule, now that we are using virtual paths

diff --git a/Engine/Plugins/Editor/AssetManagerEditor/Source/AssetManagerEditor/Private/AssetManagerEditorModule.cpp b/Engine/Plugins/Editor/AssetManagerEditor/Source/AssetManagerEditor/Private/AssetManagerEditorModule.cpp
--- a/Engine/Plugins/Editor/AssetManagerEditor/Source/AssetManagerEditor/Private/AssetManagerEditorModule.cpp
+++ b/Engine/Plugins/Editor/AssetManagerEditor/Source/AssetManagerEditor/Private/AssetManagerEditorModule.cpp
@@ -68,6 +68,8 @@
 #include "ToolMenus.h"
 #include "Toolkits/AssetEditorToolkitMenuContext.h"
 #include "ContentBrowserMenuContexts.h"
+#include "IContentBrowserDataModule.h"
+#include "ContentBrowserDataSubsystem.h"
 
 #define LOCTEXT_NAMESPACE "AssetManagerEditor"
 
@@ -788,7 +790,17 @@ TArray<FName> FAssetManagerEditorModule::GetContentBrowserSelectedAssetPackages(
 
 	if (GetSelectionDelegate.IsBound())
 	{
-		GetSelectionDelegate.Execute(SelectedAssets, SelectedPaths);
+		TArray<FString> SelectedVirtualPaths;
+		GetSelectionDelegate.Execute(SelectedAssets, SelectedVirtualPaths);
+
+		for (const FString& VirtualPath : SelectedVirtualPaths)
+		{
+			FString InvariantPath;
+			if (IContentBrowserDataModule::Get().GetSubsystem()->TryConvertVirtualPath(VirtualPath, InvariantPath) == EContentBrowserPathType::Internal)
+			{
+				SelectedPaths.Add(InvariantPath);
+			}
+		}
 	}
 
 	GetAssetDataInPaths(SelectedPaths, SelectedAssets);
