PR #8741: UE5 FPrimaryAssetId selector respects DisplayThumbnail metadata. (Contributed by [REDACTED_PERSON])

diff --git a/Engine/Plugins/Editor/AssetManagerEditor/Source/AssetManagerEditor/Private/PrimaryAssetIdCustomization.cpp b/Engine/Plugins/Editor/AssetManagerEditor/Source/AssetManagerEditor/Private/PrimaryAssetIdCustomization.cpp
--- a/Engine/Plugins/Editor/AssetManagerEditor/Source/AssetManagerEditor/Private/PrimaryAssetIdCustomization.cpp
+++ b/Engine/Plugins/Editor/AssetManagerEditor/Source/AssetManagerEditor/Private/PrimaryAssetIdCustomization.cpp
@@ -52,53 +52,56 @@ void FPrimaryAssetIdCustomization::CustomizeHeader(TSharedRef<class IPropertyHan
 		}
 	}
 
-	// Can the field be cleared
-	const bool bAllowClear = !(StructPropertyHandle->GetMetaDataProperty()->PropertyFlags & CPF_NoClear);
+	bool bShowThumbnail = true;
+	if (StructPropertyHandle->HasMetaData("DisplayThumbnail"))
+	{
+		const FString& ShowThumbnailString = StructPropertyHandle->GetMetaData("DisplayThumbnail");
+		bShowThumbnail = ShowThumbnailString.Len() == 0 || ShowThumbnailString == TEXT("true");
+	}
 
-	int32 ThumbnailSize = 64;
-	AssetThumbnail = MakeShareable( new FAssetThumbnail( FAssetData(), ThumbnailSize, ThumbnailSize, StructCustomizationUtils.GetThumbnailPool() ) );
-	UpdateThumbnail();
+	TSharedRef<SHorizontalBox> ValueBox = SNew(SHorizontalBox);
 
-	HeaderRow
-	.NameContent()
-	[
-		InStructPropertyHandle->CreatePropertyNameWidget()
-	]
-	.ValueContent()
-	.MinDesiredWidth(250.0f)
-	.MaxDesiredWidth(0.0f)
-	[
-		SNew(SHorizontalBox)
-		+SHorizontalBox::Slot()
-		.AutoWidth()
-		.VAlign(VAlign_Center)
-		.Padding(0, 0, 4, 0)
-		[
-			SNew(SOverlay)
-			+SOverlay::Slot()
-			.Padding(1)
+	if (bShowThumbnail)
+	{
+		const int32 ThumbnailSize = 64;
+		AssetThumbnail = MakeShareable(new FAssetThumbnail(FAssetData(), ThumbnailSize, ThumbnailSize, StructCustomizationUtils.GetThumbnailPool()));
+		UpdateThumbnail();
+
+		ValueBox->AddSlot()
+			.AutoWidth()
+			.VAlign(VAlign_Center)
+			.Padding(0, 0, 4, 0)
 			[
-				SAssignNew(ThumbnailBorder, SBorder)
-				.Padding(0)
-				.BorderImage(FStyleDefaults::GetNoBrush())
-				.OnMouseDoubleClick(this, &FPrimaryAssetIdCustomization::OnAssetThumbnailDoubleClick)
+				SNew(SOverlay)
+				+ SOverlay::Slot()
+				.Padding(1)
 				[
-					SNew(SBox)
-					.WidthOverride(ThumbnailSize)
-					.HeightOverride(ThumbnailSize)
+					SAssignNew(ThumbnailBorder, SBorder)
+					.Padding(0)
+					.BorderImage(FStyleDefaults::GetNoBrush())
+					.OnMouseDoubleClick(this, &FPrimaryAssetIdCustomization::OnAssetThumbnailDoubleClick)
 					[
-						AssetThumbnail->MakeThumbnailWidget()
+						SNew(SBox)
+						.WidthOverride(ThumbnailSize)
+						.HeightOverride(ThumbnailSize)
+						[
+							AssetThumbnail->MakeThumbnailWidget()
+						]
 					]
 				]
-			]
-			+ SOverlay::Slot()
-			[
-				SNew(SImage)
-				.Image(this, &FPrimaryAssetIdCustomization::GetThumbnailBorder)
-				.Visibility(EVisibility::SelfHitTestInvisible)
-			]
-		]
-		+ SHorizontalBox::Slot()
+				+ SOverlay::Slot()
+				[
+					SNew(SImage)
+					.Image(this, &FPrimaryAssetIdCustomization::GetThumbnailBorder)
+					.Visibility(EVisibility::SelfHitTestInvisible)
+				]
+			];
+	}
+
+	// Can the field be cleared
+	const bool bAllowClear = !(StructPropertyHandle->GetMetaDataProperty()->PropertyFlags & CPF_NoClear);
+
+	ValueBox->AddSlot()
 		.FillWidth(1.0f)
 		.VAlign(VAlign_Center)
 		[
@@ -106,25 +109,38 @@ void FPrimaryAssetIdCustomization::CustomizeHeader(TSharedRef<class IPropertyHan
 				FOnGetPrimaryAssetDisplayText::CreateSP(this, &FPrimaryAssetIdCustomization::GetDisplayText),
 				FOnSetPrimaryAssetId::CreateSP(this, &FPrimaryAssetIdCustomization::OnIdSelected),
 				bAllowClear, AllowedTypes)
-		]
-		+SHorizontalBox::Slot()
+		];
+
+	ValueBox->AddSlot()
 		.AutoWidth()
 		.VAlign(VAlign_Center)
 		[
 			PropertyCustomizationHelpers::MakeUseSelectedButton(FSimpleDelegate::CreateSP(this, &FPrimaryAssetIdCustomization::OnUseSelected))
-		]
-		+SHorizontalBox::Slot()
+		];
+
+	ValueBox->AddSlot()
 		.AutoWidth()
 		.VAlign(VAlign_Center)
 		[
 			PropertyCustomizationHelpers::MakeBrowseButton(FSimpleDelegate::CreateSP(this, &FPrimaryAssetIdCustomization::OnBrowseTo))
-		]
-		+SHorizontalBox::Slot()
+		];
+
+	ValueBox->AddSlot()
 		.AutoWidth()
 		.VAlign(VAlign_Center)
 		[
 			PropertyCustomizationHelpers::MakeClearButton(FSimpleDelegate::CreateSP(this, &FPrimaryAssetIdCustomization::OnClear))
-		]
+		];
+	HeaderRow
+	.NameContent()
+	[
+		InStructPropertyHandle->CreatePropertyNameWidget()
+	]
+	.ValueContent()
+	.MinDesiredWidth(250.0f)
+	.MaxDesiredWidth(0.0f)
+	[
+		ValueBox
 	];
 }
 
@@ -178,16 +194,17 @@ FPrimaryAssetId FPrimaryAssetIdCustomization::GetCurrentPrimaryAssetId() const
 
 void FPrimaryAssetIdCustomization::UpdateThumbnail()
 {
-	check(AssetThumbnail.IsValid());
-
-	FAssetData AssetData;
-	FPrimaryAssetId PrimaryAssetId = GetCurrentPrimaryAssetId();
-	if (PrimaryAssetId.IsValid())
+	if (AssetThumbnail.IsValid())
 	{
-		UAssetManager::Get().GetPrimaryAssetData(PrimaryAssetId, AssetData);
-	}
+		FAssetData AssetData;
+		FPrimaryAssetId PrimaryAssetId = GetCurrentPrimaryAssetId();
+		if (PrimaryAssetId.IsValid())
+		{
+			UAssetManager::Get().GetPrimaryAssetData(PrimaryAssetId, AssetData);
+		}
 
-	AssetThumbnail->SetAsset(AssetData);
+		AssetThumbnail->SetAsset(AssetData);
+	}
 }
 
 void FPrimaryAssetIdCustomization::OnBrowseTo()
