Add a timestamp for the AssetRegistry.bin that's in use in the Asset Auditor. This is to help when people don't realize they have to re-cook to pick up changes.

diff --git a/Engine/Plugins/Editor/AssetManagerEditor/Source/AssetManagerEditor/Private/AssetManagerEditorModule.cpp b/Engine/Plugins/Editor/AssetManagerEditor/Source/AssetManagerEditor/Private/AssetManagerEditorModule.cpp
--- a/Engine/Plugins/Editor/AssetManagerEditor/Source/AssetManagerEditor/Private/AssetManagerEditorModule.cpp
+++ b/Engine/Plugins/Editor/AssetManagerEditor/Source/AssetManagerEditor/Private/AssetManagerEditorModule.cpp
@@ -1479,6 +1479,32 @@ const FAssetManagerEditorRegistrySource* FAssetManagerEditorModule::GetCurrentRe
 	return CurrentRegistrySource;
 }
 
+void FAssetManagerEditorRegistrySource::LoadRegistryTimestamp()
+{
+	FFileStatData TimeData = IFileManager::Get().GetStatData(*SourceFilename);
+	FDateTime UseTime = FDateTime::MinValue();
+	if (TimeData.bIsValid)
+	{
+		if (TimeData.ModificationTime != FDateTime::MinValue())
+		{
+			UseTime = TimeData.ModificationTime;
+		}
+		else
+		{
+			UseTime = TimeData.CreationTime;
+		}
+	}
+
+	if (UseTime != FDateTime::MinValue())
+	{
+		//Turn UTC into local
+		FTimespan UTCOffset = FDateTime::Now() - FDateTime::UtcNow();
+		UseTime += UTCOffset;
+
+		SourceTimestamp = UseTime.ToString(TEXT("%Y.%m.%d %h:%m %A"));
+	}
+}
+
 void FAssetManagerEditorModule::SetCurrentRegistrySource(const FString& SourceName)
 {
 	InitializeRegistrySources(false);
@@ -1489,6 +1515,8 @@ void FAssetManagerEditorModule::SetCurrentRegistrySource(const FString& SourceNa
 	{
 		CurrentRegistrySource = NewSource;
 
+		CurrentRegistrySource->SourceTimestamp = FString();
+
 		if (CurrentRegistrySource->SourceName == FAssetManagerEditorRegistrySource::CustomSourceName)
 		{
 			if (CurrentRegistrySource->HasRegistry())
@@ -1532,6 +1560,8 @@ void FAssetManagerEditorModule::SetCurrentRegistrySource(const FString& SourceNa
 			FArrayReader SerializedAssetData;
 			if (FFileHelper::LoadFileToArray(SerializedAssetData, *CurrentRegistrySource->SourceFilename))
 			{
+				CurrentRegistrySource->LoadRegistryTimestamp();
+
 				FAssetRegistryState* NewState = new FAssetRegistryState();
 				FAssetRegistrySerializationOptions Options(UE::AssetRegistry::ESerializationTarget::ForDevelopment);
 
diff --git a/Engine/Plugins/Editor/AssetManagerEditor/Source/AssetManagerEditor/Private/SAssetAuditBrowser.cpp b/Engine/Plugins/Editor/AssetManagerEditor/Source/AssetManagerEditor/Private/SAssetAuditBrowser.cpp
--- a/Engine/Plugins/Editor/AssetManagerEditor/Source/AssetManagerEditor/Private/SAssetAuditBrowser.cpp
+++ b/Engine/Plugins/Editor/AssetManagerEditor/Source/AssetManagerEditor/Private/SAssetAuditBrowser.cpp
@@ -455,6 +455,15 @@ void SAssetAuditBrowser::Construct(const FArguments& InArgs)
 						.Text(LOCTEXT("RefreshAssets", "Refresh"))
 						.OnClicked(this, &SAssetAuditBrowser::RefreshAssets)
 					]
+
+					+ SHorizontalBox::Slot()
+					.Padding(6.0f)
+					.VAlign(VAlign_Center)
+					.AutoWidth()
+					[
+						SAssignNew(RegistrySourceTimeText, STextBlock)
+						.ToolTipText(LOCTEXT("RegistryTime_TT", "The timestamp for the AssetRegistry.bin file. Changes to assets are not picked up until a new cook."))
+					]
 					
 					+ SHorizontalBox::Slot()
 					.FillWidth(1.0f)
@@ -874,6 +883,23 @@ void SAssetAuditBrowser::SetCurrentRegistrySource(const FAssetManagerEditorRegis
 {
 	CurrentRegistrySource = RegistrySource;
 
+	if (CurrentRegistrySource->bIsEditor)
+	{
+		RegistrySourceTimeText->SetText(LOCTEXT("AssetRegistryTimestampEditor", "Preview data (cook for actual data)"));
+	}
+	else
+	{
+		
+		if (CurrentRegistrySource->SourceTimestamp.Len() == 0)
+		{
+			RegistrySourceTimeText->SetText(LOCTEXT("AssetRegistryTimestampError", "Unable to get AssetRegistry.bin timestamp (recook?)"));
+		}
+		else
+		{
+			RegistrySourceTimeText->SetText(FText::Format(LOCTEXT("AssetRegistryTimestampFormat", "AssetRegistry.bin from {0}"), FText::FromString(CurrentRegistrySource->SourceTimestamp)));
+		}
+	}
+
 	// Refresh dropdown
 	TArray<const FAssetManagerEditorRegistrySource*> AvailableSources;
 	IAssetManagerEditorModule::Get().GetAvailableRegistrySources(AvailableSources);
diff --git a/Engine/Plugins/Editor/AssetManagerEditor/Source/AssetManagerEditor/Private/SAssetAuditBrowser.h b/Engine/Plugins/Editor/AssetManagerEditor/Source/AssetManagerEditor/Private/SAssetAuditBrowser.h
--- a/Engine/Plugins/Editor/AssetManagerEditor/Source/AssetManagerEditor/Private/SAssetAuditBrowser.h
+++ b/Engine/Plugins/Editor/AssetManagerEditor/Source/AssetManagerEditor/Private/SAssetAuditBrowser.h
@@ -14,6 +14,7 @@
 
 class FUICommandList;
 class SMenuAnchor;
+class STextBlock;
 
 //////////////////////////////////////////////////////////////////////////
 // SAssetAuditBrowser
@@ -145,6 +146,12 @@ protected:
 	/** List of valid registry sources */
 	TArray<TSharedPtr<FString>> SourceComboList;
 
+	/**
+	* Text block where we put the timestamp for the asset registry so the user knows
+	* if they are looking at super out of date data.
+	*/
+	TSharedPtr<STextBlock> RegistrySourceTimeText;
+
 	/** Current TargetPlatform and registry state, may be null! */
 	const FAssetManagerEditorRegistrySource* CurrentRegistrySource;
 
diff --git a/Engine/Plugins/Editor/AssetManagerEditor/Source/AssetManagerEditor/Public/AssetManagerEditorModule.h b/Engine/Plugins/Editor/AssetManagerEditor/Source/AssetManagerEditor/Public/AssetManagerEditorModule.h
--- a/Engine/Plugins/Editor/AssetManagerEditor/Source/AssetManagerEditor/Public/AssetManagerEditorModule.h
+++ b/Engine/Plugins/Editor/AssetManagerEditor/Source/AssetManagerEditor/Public/AssetManagerEditorModule.h
@@ -59,6 +59,12 @@ struct ASSETMANAGEREDITOR_API FAssetManagerEditorRegistrySource
 	/** Filename registry was loaded from */
 	FString SourceFilename;
 
+	/** 
+	* Timestamp for the registry we loaded from, if from a file. Modified time if the OS supports it,
+	* otherwise creation time. If we can't get anything or not from a file, empty string. In local timezone.
+	*/
+	FString SourceTimestamp;
+
 	/** Target platform for this state, may be null */
 	ITargetPlatform* TargetPlatform;
 
@@ -190,6 +196,9 @@ public:
 		AssetRegistryPrivate = nullptr;
 	}
 
+	// grabs the SourceTimestamp from the asset registry file changed times.
+	void LoadRegistryTimestamp();
+
 	// Functions to forward on to the RegistryState, either the owned one or the editor's global State.
 
 	FAssetData GetAssetByObjectPath(FName ObjectPath) const
