Fix for timestamp in asset auditor using month instead of minute.
Fix for timetamp getting incorrectly cleared when a custom audit is cancelled.

diff --git a/Engine/Plugins/Editor/AssetManagerEditor/Source/AssetManagerEditor/Private/AssetManagerEditorModule.cpp b/Engine/Plugins/Editor/AssetManagerEditor/Source/AssetManagerEditor/Private/AssetManagerEditorModule.cpp
--- a/Engine/Plugins/Editor/AssetManagerEditor/Source/AssetManagerEditor/Private/AssetManagerEditorModule.cpp
+++ b/Engine/Plugins/Editor/AssetManagerEditor/Source/AssetManagerEditor/Private/AssetManagerEditorModule.cpp
@@ -1501,7 +1501,7 @@ void FAssetManagerEditorRegistrySource::LoadRegistryTimestamp()
 		FTimespan UTCOffset = FDateTime::Now() - FDateTime::UtcNow();
 		UseTime += UTCOffset;
 
-		SourceTimestamp = UseTime.ToString(TEXT("%Y.%m.%d %h:%m %A"));
+		SourceTimestamp = UseTime.ToString(TEXT("%Y.%m.%d %h:%M %A"));
 	}
 }
 
@@ -1515,10 +1515,10 @@ void FAssetManagerEditorModule::SetCurrentRegistrySource(const FString& SourceNa
 	{
 		CurrentRegistrySource = NewSource;
 
-		CurrentRegistrySource->SourceTimestamp = FString();
-
 		if (CurrentRegistrySource->SourceName == FAssetManagerEditorRegistrySource::CustomSourceName)
 		{
+			CurrentRegistrySource->SourceTimestamp = FString();
+
 			if (CurrentRegistrySource->HasRegistry())
 			{
 				CurrentRegistrySource->ClearRegistry();
@@ -1556,6 +1556,8 @@ void FAssetManagerEditorModule::SetCurrentRegistrySource(const FString& SourceNa
 
 		if (!CurrentRegistrySource->HasRegistry() && !CurrentRegistrySource->SourceFilename.IsEmpty() && !CurrentRegistrySource->bIsEditor)
 		{
+			CurrentRegistrySource->SourceTimestamp = FString();
+
 			bool bLoaded = false;
 			FArrayReader SerializedAssetData;
 			if (FFileHelper::LoadFileToArray(SerializedAssetData, *CurrentRegistrySource->SourceFilename))
