Fix for asset audit ui text moving the combo box off the edge of the window making it inaccessible.

diff --git a/Engine/Plugins/Editor/AssetManagerEditor/Source/AssetManagerEditor/Private/AssetManagerEditorModule.cpp b/Engine/Plugins/Editor/AssetManagerEditor/Source/AssetManagerEditor/Private/AssetManagerEditorModule.cpp
--- a/Engine/Plugins/Editor/AssetManagerEditor/Source/AssetManagerEditor/Private/AssetManagerEditorModule.cpp
+++ b/Engine/Plugins/Editor/AssetManagerEditor/Source/AssetManagerEditor/Private/AssetManagerEditorModule.cpp
@@ -497,6 +497,7 @@ void FAssetManagerEditorModule::StartupModule()
 			.SetTooltipText(LOCTEXT("AssetAuditTooltip", "Open Asset Audit window, allows viewing detailed information about assets."))
 			.SetGroup(WorkspaceMenu::GetMenuStructure().GetDeveloperToolsAuditCategory())
 			.SetIcon(FSlateIcon(FAppStyle::GetAppStyleSetName(), "LevelEditor.Audit"));
+		FGlobalTabmanager::Get()->RegisterDefaultTabWindowSize(AssetAuditTabName, FVector2D(1080, 600));
 
 		FGlobalTabmanager::Get()->RegisterNomadTabSpawner(ReferenceViewerTabName, FOnSpawnTab::CreateRaw(this, &FAssetManagerEditorModule::SpawnReferenceViewerTab))
 			.SetDisplayName(LOCTEXT("ReferenceViewerTitle", "Reference Viewer"))
diff --git a/Engine/Plugins/Editor/AssetManagerEditor/Source/AssetManagerEditor/Private/SAssetAuditBrowser.cpp b/Engine/Plugins/Editor/AssetManagerEditor/Source/AssetManagerEditor/Private/SAssetAuditBrowser.cpp
--- a/Engine/Plugins/Editor/AssetManagerEditor/Source/AssetManagerEditor/Private/SAssetAuditBrowser.cpp
+++ b/Engine/Plugins/Editor/AssetManagerEditor/Source/AssetManagerEditor/Private/SAssetAuditBrowser.cpp
@@ -456,34 +456,42 @@ void SAssetAuditBrowser::Construct(const FArguments& InArgs)
 						.OnClicked(this, &SAssetAuditBrowser::RefreshAssets)
 					]
 
-					+ SHorizontalBox::Slot()
-					.Padding(6.0f)
-					.VAlign(VAlign_Center)
-					.AutoWidth()
-					[
-						SAssignNew(RegistrySourceTimeText, STextBlock)
-						.ToolTipText(LOCTEXT("RegistryTime_TT", "The timestamp for the AssetRegistry.bin file. Changes to assets are not picked up until a new cook."))
-					]
-					
-					+ SHorizontalBox::Slot()
+					+SHorizontalBox::Slot()
 					.FillWidth(1.0f)
 					.HAlign(HAlign_Fill)
 					[
-						SNew(SHorizontalBox)						
+						SNew(SHorizontalBox)
 						+ SHorizontalBox::Slot()
-						.FillWidth(1.0f)
-						.HAlign(HAlign_Right)
 						.VAlign(VAlign_Center)
+						.HAlign(HAlign_Fill)
+						.FillWidth(1.0f)
 						[
-							SNew(STextBlock)
-							.ToolTipText(LOCTEXT("Platform_Tooltip", "Select which platform to display data for. Platforms are only available if a cooked AssetRegistry.bin is available in Saved/Cooked/Platform or Build/Platform."))
-							.Text(LOCTEXT("PlatformLabel", "Selected Platform: "))
+							SNew(SHorizontalBox)
+							+ SHorizontalBox::Slot()
+							.Padding(6.0f)
+							.VAlign(VAlign_Center)
+							.HAlign(HAlign_Left)
+							.FillWidth(1.0f)
+							[
+								SAssignNew(RegistrySourceTimeText, STextBlock)
+								.ToolTipText(LOCTEXT("RegistryTime_TT", "If viewing cooked data, the timestamp for the AssetRegistry.bin file. Changes to assets are not picked up until a new cook. Otherwise, viewing preview data from the editor."))
+							]
+							+ SHorizontalBox::Slot()
+							.Padding(6.0f)
+							.AutoWidth()
+							.HAlign(HAlign_Right)
+							.VAlign(VAlign_Center)
+							[
+								SNew(STextBlock)
+								.ToolTipText(LOCTEXT("Platform_Tooltip", "Select which platform to display data for. Platforms are only available if a cooked AssetRegistry.bin is available in Saved/Cooked/Platform or Build/Platform."))
+								.Text(LOCTEXT("PlatformLabel", "Selected Platform: "))
+							]
 						]
 						+ SHorizontalBox::Slot()
 						.AutoWidth()
 						.HAlign(HAlign_Right)
 						[
-							SNew(SComboBox<TSharedPtr<FString>>)							
+							SNew(SComboBox<TSharedPtr<FString>>)
 							.OptionsSource(&SourceComboList)
 							.OnGenerateWidget(this, &SAssetAuditBrowser::GenerateSourceComboItem)
 							.OnSelectionChanged(this, &SAssetAuditBrowser::HandleSourceComboChanged)
@@ -885,7 +893,7 @@ void SAssetAuditBrowser::SetCurrentRegistrySource(const FAssetManagerEditorRegis
 
 	if (CurrentRegistrySource->bIsEditor)
 	{
-		RegistrySourceTimeText->SetText(LOCTEXT("AssetRegistryTimestampEditor", "Preview data (cook for actual data)"));
+		RegistrySourceTimeText->SetText(LOCTEXT("AssetRegistryTimestampEditor", "Preview data"));
 	}
 	else
 	{
