Fixed the edit menu option not working when a primary asset ID was selected in the reference viewer

diff --git a/Engine/Plugins/Editor/AssetManagerEditor/Source/AssetManagerEditor/Private/AssetManagerEditorModule.cpp b/Engine/Plugins/Editor/AssetManagerEditor/Source/AssetManagerEditor/Private/AssetManagerEditorModule.cpp
--- a/Engine/Plugins/Editor/AssetManagerEditor/Source/AssetManagerEditor/Private/AssetManagerEditorModule.cpp
+++ b/Engine/Plugins/Editor/AssetManagerEditor/Source/AssetManagerEditor/Private/AssetManagerEditorModule.cpp
@@ -1048,17 +1048,35 @@ void FAssetManagerEditorModule::OnMarkPackageDirty(UPackage* Pkg, bool bWasDirty
 
 void FAssetManagerEditorModule::OnEditAssetIdentifiers(TArray<FAssetIdentifier> AssetIdentifiers)
 {
-	// Handle default package behavior
+	UAssetManager& AssetManager = UAssetManager::Get();
+
+	// Determine which packages to load
 	TArray<FAssetData> AssetsToLoad;
 	for (FAssetIdentifier AssetIdentifier : AssetIdentifiers)
 	{
 		if (AssetIdentifier.IsPackage())
 		{
-			TArray<FAssetData> Assets;
+			// Directly a package to load
 			AssetRegistry->GetAssetsByPackageName(AssetIdentifier.PackageName, AssetsToLoad);
 		}
+		else
+		{	
+			// If it's a primary asset ID, resolve it to a package to load
+			FPrimaryAssetId AssetId = AssetIdentifier.GetPrimaryAssetId();
+			if (AssetId.IsValid())
+			{
+				//@TODO: We probably want to call UAssetManager::GetNameData with bCheckRedirector=false but that's protected
+				FSoftObjectPath AssetPath = AssetManager.GetPrimaryAssetPath(AssetId);
+				FAssetData AssetData;
+				if (AssetManager.GetAssetDataForPath(AssetPath, /*out*/ AssetData))
+				{
+					AssetsToLoad.Add(AssetData);
+				}
+			}
+		}
 	}
 
+	// Open the editor(s)
 	if (AssetsToLoad.Num() > 0)
 	{
 		FScopedSlowTask SlowTask(0, LOCTEXT("LoadingSelectedObject", "Editing assets..."));
