Added a project setting to allow games to default bIsShowSearchableNames on for the Reference Viewer

diff --git a/Engine/Plugins/Editor/AssetManagerEditor/Source/AssetManagerEditor/Private/ReferenceViewer/EdGraph_ReferenceViewer.cpp b/Engine/Plugins/Editor/AssetManagerEditor/Source/AssetManagerEditor/Private/ReferenceViewer/EdGraph_ReferenceViewer.cpp
--- a/Engine/Plugins/Editor/AssetManagerEditor/Source/AssetManagerEditor/Private/ReferenceViewer/EdGraph_ReferenceViewer.cpp
+++ b/Engine/Plugins/Editor/AssetManagerEditor/Source/AssetManagerEditor/Private/ReferenceViewer/EdGraph_ReferenceViewer.cpp
@@ -13,6 +13,7 @@
 #include "CollectionManagerModule.h"
 #include "AssetManagerEditorModule.h"
 #include "Engine/AssetManager.h"
+#include "Settings/EditorProjectSettings.h"
 
 UEdGraph_ReferenceViewer::UEdGraph_ReferenceViewer(const FObjectInitializer& ObjectInitializer)
 	: Super(ObjectInitializer)
@@ -20,7 +21,7 @@ UEdGraph_ReferenceViewer::UEdGraph_ReferenceViewer(const FObjectInitializer& Obj
 	AssetThumbnailPool = MakeShareable( new FAssetThumbnailPool(1024) );
 
 	MaxSearchDepth = 1;
-	MaxSearchBreadth = 20;
+	MaxSearchBreadth = GetDefault<UEditorProjectAppearanceSettings>()->ReferenceViewerDefaultMaxSearchBreadth;
 
 	bLimitSearchDepth = true;
 	bLimitSearchBreadth = true;
@@ -28,7 +29,7 @@ UEdGraph_ReferenceViewer::UEdGraph_ReferenceViewer(const FObjectInitializer& Obj
 	bIsShowHardReferences = true;
 	bIsShowEditorOnlyReferences = true;
 	bIsShowManagementReferences = false;
-	bIsShowSearchableNames = false;
+	bIsShowSearchableNames = GetDefault<UEditorProjectAppearanceSettings>()->ShowSearchableNames == EReferenceViewerSettingMode::ShowByDefault;
 	bIsShowNativePackages = false;
 	bIsShowReferencers = true;
 	bIsShowDependencies = true;
diff --git a/Engine/Source/Editor/UnrealEd/Public/Settings/EditorProjectSettings.h b/Engine/Source/Editor/UnrealEd/Public/Settings/EditorProjectSettings.h
--- a/Engine/Source/Editor/UnrealEd/Public/Settings/EditorProjectSettings.h
+++ b/Engine/Source/Editor/UnrealEd/Public/Settings/EditorProjectSettings.h
@@ -38,6 +38,19 @@ enum class EDefaultLocationUnit : uint8
 	Invalid
 };
 
+UENUM()
+enum class EReferenceViewerSettingMode : uint8
+{
+	// Use the editor default setting
+	NoPreference,
+
+	// Show this kind of reference by default (it can be toggled off in the reference viewer)
+	ShowByDefault,
+
+	// Hide this kind of reference by default (it can be toggled back on in the reference viewer)
+	HideByDefault
+};
+
 /**
  * Editor project appearance settings. Stored in default config, per-project
  */
@@ -82,6 +95,14 @@ public:
 	UPROPERTY(EditAnywhere, config, Category=Units, AdvancedDisplay, meta=(DisplayName="Force", Tooltip="Choose the units in which to display forces.", ValidEnumValues="Newtons, PoundsForce, KilogramsForce"))
 	EUnit ForceUnits;
 
+	// Should the Reference Viewer have 'Show Searchable Names' checked by default when opened in this project
+	UPROPERTY(EditAnywhere, config, Category=ReferenceViewer)
+	EReferenceViewerSettingMode ShowSearchableNames;
+
+	// The default maximum search breadth for the reference viewer when opened
+	UPROPERTY(EditAnywhere, config, Category = ReferenceViewer, meta=(DisplayName="Default Max Search Breadth", ClampMin=1, ClampMax=1000, UIMin=1, UIMax=50))
+	int32 ReferenceViewerDefaultMaxSearchBreadth = 20;
+
 public:
 	/** Deprecated properties that didn't live very long */
 
