ReferenceViewer Enhancements Part 1

- Updates the header style (UE5, emphasizes type)
- Shows nodes that may show up in more than one place (duplicates)
- Separates Depth control into Referencer Depth and Dependency Depth
- Comments with Package Path have been removed entirely (the tooltip remains)
- Rewritten algorithms to generate graph nodes that remove code duplication and will make adding new features easier
- Consolidates filters within the GetSortedLinks function
- adds temporary CVar that enables the ability to toggle between the old and new reference node layout algorithms.  The depricated algorithms and CVar will be removed once the rest of phase one is completed and the new layout has undergone more testing.

diff --git a/Engine/Plugins/Editor/AssetManagerEditor/Source/AssetManagerEditor/Private/ReferenceViewer/EdGraphNode_Reference.cpp b/Engine/Plugins/Editor/AssetManagerEditor/Source/AssetManagerEditor/Private/ReferenceViewer/EdGraphNode_Reference.cpp
--- a/Engine/Plugins/Editor/AssetManagerEditor/Source/AssetManagerEditor/Private/ReferenceViewer/EdGraphNode_Reference.cpp
+++ b/Engine/Plugins/Editor/AssetManagerEditor/Source/AssetManagerEditor/Private/ReferenceViewer/EdGraphNode_Reference.cpp
@@ -4,7 +4,11 @@
 #include "GenericPlatform/GenericPlatformFile.h"
 #include "EdGraph/EdGraphPin.h"
 #include "HAL/PlatformFileManager.h"
+#include "Styling/SlateIconFinder.h"
 #include "SReferenceViewer.h"
+#include "AssetToolsModule.h"
+#include "IAssetTools.h"
+#include "IAssetTypeActions.h"
 
 #define LOCTEXT_NAMESPACE "ReferenceViewer"
 
@@ -21,9 +25,10 @@ UEdGraphNode_Reference::UEdGraphNode_Reference(const FObjectInitializer& ObjectI
 	bIsPrimaryAsset = false;
 	bUsesThumbnail = false;
 	bAllowThumbnail = true;
+	AssetTypeColor = FLinearColor(0.55f, 0.55f, 0.55f);
 }
 
-void UEdGraphNode_Reference::SetupReferenceNode(const FIntPoint& NodeLoc, const TArray<FAssetIdentifier>& NewIdentifiers, const FAssetData& InAssetData, bool bInAllowThumbnail)
+void UEdGraphNode_Reference::SetupReferenceNode(const FIntPoint& NodeLoc, const TArray<FAssetIdentifier>& NewIdentifiers, const FAssetData& InAssetData, bool bInAllowThumbnail, bool bInIsADuplicate)
 {
 	check(NewIdentifiers.Num() > 0);
 
@@ -34,9 +39,21 @@ void UEdGraphNode_Reference::SetupReferenceNode(const FIntPoint& NodeLoc, const
 	const FAssetIdentifier& First = NewIdentifiers[0];
 	FString MainAssetName = InAssetData.AssetName.ToString();
 
+	FAssetToolsModule& AssetToolsModule = FModuleManager::LoadModuleChecked<FAssetToolsModule>(TEXT("AssetTools"));	
+	if (UClass* AssetClass = InAssetData.GetClass())
+	{
+		TWeakPtr<IAssetTypeActions> AssetTypeActions = AssetToolsModule.Get().GetAssetTypeActionsForClass(InAssetData.GetClass());
+		if(AssetTypeActions.IsValid())
+		{
+			AssetTypeColor = AssetTypeActions.Pin()->GetTypeColor();
+		}
+	}
+	AssetBrush = FSlateIcon("EditorStyle", FName( *("ClassIcon." + InAssetData.AssetClass.ToString())));
+
 	bIsCollapsed = false;
 	bIsPackage = true;
 	bAllowThumbnail = bInAllowThumbnail;
+	bIsADuplicate = bInIsADuplicate;
 	
 	FPrimaryAssetId PrimaryAssetID = NewIdentifiers[0].GetPrimaryAssetId();
 	if (PrimaryAssetID.IsValid())
@@ -53,19 +70,13 @@ void UEdGraphNode_Reference::SetupReferenceNode(const FIntPoint& NodeLoc, const
 
 	if (NewIdentifiers.Num() == 1 )
 	{
-		if (bIsPackage)
-		{
-			NodeComment = First.PackageName.ToString();
-		}
-
 		static const FName NAME_ActorLabel(TEXT("ActorLabel"));
-
-		NodeTitle = FText::FromString(MainAssetName);
+		
+		NodeTitle = FText::FromString(FString::Printf(TEXT("%s\n%s"), *InAssetData.AssetName.ToString(), *InAssetData.AssetClass.ToString()));
 		InAssetData.GetTagValue(NAME_ActorLabel, NodeTitle);
 	}
 	else
 	{
-		NodeComment = FText::Format(LOCTEXT("ReferenceNodeMultiplePackagesTitle", "{0} nodes"), FText::AsNumber(NewIdentifiers.Num())).ToString();
 		NodeTitle = FText::Format(LOCTEXT("ReferenceNodeMultiplePackagesComment", "{0} and {1} others"), FText::FromString(MainAssetName), FText::AsNumber(NewIdentifiers.Num() - 1));
 	}
 	
@@ -81,9 +92,8 @@ void UEdGraphNode_Reference::SetReferenceNodeCollapsed(const FIntPoint& NodeLoc,
 	Identifiers.Empty();
 	bIsCollapsed = true;
 	bUsesThumbnail = false;
-	NodeComment = FText::Format(LOCTEXT("ReferenceNodeCollapsedMessage", "{0} other nodes"), FText::AsNumber(InNumReferencesExceedingMax)).ToString();
 
-	NodeTitle = LOCTEXT("ReferenceNodeCollapsedTitle", "Collapsed nodes");
+	NodeTitle = FText::Format( LOCTEXT("ReferenceNodeCollapsedTitle", "{0} Collapsed nodes"), FText::AsNumber(InNumReferencesExceedingMax));
 	CacheAssetData(FAssetData());
 	AllocateDefaultPins();
 }
@@ -144,7 +154,7 @@ FLinearColor UEdGraphNode_Reference::GetNodeTitleColor() const
 	}
 	else if (bIsPackage)
 	{
-		return FLinearColor(0.4f, 0.62f, 1.0f);
+		return AssetTypeColor;
 	}
 	else if (bIsCollapsed)
 	{
@@ -170,6 +180,12 @@ FText UEdGraphNode_Reference::GetTooltipText() const
 	return FText::FromString(TooltipString);
 }
 
+FSlateIcon UEdGraphNode_Reference::GetIconAndTint(FLinearColor& OutColor) const
+{
+	OutColor = AssetTypeColor;
+	return AssetBrush;
+}
+
 void UEdGraphNode_Reference::AllocateDefaultPins()
 {
 	ReferencerPin = CreatePin( EEdGraphPinDirection::EGPD_Input, NAME_None, NAME_None);
diff --git a/Engine/Plugins/Editor/AssetManagerEditor/Source/AssetManagerEditor/Private/ReferenceViewer/EdGraph_ReferenceViewer.cpp b/Engine/Plugins/Editor/AssetManagerEditor/Source/AssetManagerEditor/Private/ReferenceViewer/EdGraph_ReferenceViewer.cpp
--- a/Engine/Plugins/Editor/AssetManagerEditor/Source/AssetManagerEditor/Private/ReferenceViewer/EdGraph_ReferenceViewer.cpp
+++ b/Engine/Plugins/Editor/AssetManagerEditor/Source/AssetManagerEditor/Private/ReferenceViewer/EdGraph_ReferenceViewer.cpp
@@ -15,12 +15,35 @@
 #include "Engine/AssetManager.h"
 #include "Settings/EditorProjectSettings.h"
 
+ReferenceNodeInfo::ReferenceNodeInfo(const FAssetIdentifier& InAssetId, bool InbReferencers)
+	: AssetId(InAssetId)
+	, bReferencers(InbReferencers)
+	, OverflowCount(0)
+	, ChildProvisionSize(0)
+{}
+
+bool ReferenceNodeInfo::IsFirstParent(const FAssetIdentifier& InParentId) const
+{
+	return Parents.IsEmpty() || Parents[0] == InParentId;
+}
+
+bool ReferenceNodeInfo::IsADuplicate() const
+{
+	return Parents.Num() > 1;
+}
+
+int32 ReferenceNodeInfo::ProvisionSize(const FAssetIdentifier& InParentId) const
+{
+	return IsFirstParent(InParentId) ? ChildProvisionSize : 1;
+}
+
 UEdGraph_ReferenceViewer::UEdGraph_ReferenceViewer(const FObjectInitializer& ObjectInitializer)
 	: Super(ObjectInitializer)
 {
 	AssetThumbnailPool = MakeShareable( new FAssetThumbnailPool(1024) );
 
-	MaxSearchDepth = 1;
+	MaxSearchDependencyDepth = 1;
+	MaxSearchReferencerDepth = 1;
 	MaxSearchBreadth = GetDefault<UEditorProjectAppearanceSettings>()->ReferenceViewerDefaultMaxSearchBreadth;
 
 	bLimitSearchDepth = true;
@@ -35,6 +58,8 @@ UEdGraph_ReferenceViewer::UEdGraph_ReferenceViewer(const FObjectInitializer& Obj
 	bIsShowDependencies = true;
 	bIsShowFilteredPackagesOnly = false;
 	bIsCompactMode = false;
+	bIsShowDuplicates = true;
+	bUseNodeInfos = true;
 }
 
 void UEdGraph_ReferenceViewer::BeginDestroy()
@@ -144,6 +169,11 @@ bool UEdGraph_ReferenceViewer::IsCompactMode() const
 	return bIsCompactMode;
 }
 
+bool UEdGraph_ReferenceViewer::IsShowDuplicates() const
+{
+	return bIsShowDuplicates;
+}
+
 bool UEdGraph_ReferenceViewer::IsShowEditorOnlyReferences() const
 {
 	return bIsShowEditorOnlyReferences;
@@ -204,6 +234,11 @@ void UEdGraph_ReferenceViewer::SetCompactModeEnabled(bool newEnabled)
 	bIsCompactMode = newEnabled;
 }
 
+void UEdGraph_ReferenceViewer::SetShowDuplicatesEnabled(bool newEnabled)
+{
+	bIsShowDuplicates = newEnabled;
+}
+
 void UEdGraph_ReferenceViewer::SetShowEditorOnlyReferencesEnabled(bool newEnabled)
 {
 	bIsShowEditorOnlyReferences = newEnabled;
@@ -234,19 +269,29 @@ void UEdGraph_ReferenceViewer::SetShowDependencies(const bool newEnabled)
 	bIsShowDependencies = newEnabled;
 }
 
-int32 UEdGraph_ReferenceViewer::GetSearchDepthLimit() const
+int32 UEdGraph_ReferenceViewer::GetSearchDependencyDepthLimit() const
 {
-	return MaxSearchDepth;
+	return MaxSearchDependencyDepth;
 }
 
-int32 UEdGraph_ReferenceViewer::GetSearchBreadthLimit() const
+void UEdGraph_ReferenceViewer::SetSearchDependencyDepthLimit(int32 NewDepthLimit)
 {
-	return MaxSearchBreadth;
+	MaxSearchDependencyDepth = NewDepthLimit;
 }
 
-void UEdGraph_ReferenceViewer::SetSearchDepthLimit(int32 NewDepthLimit)
+int32 UEdGraph_ReferenceViewer::GetSearchReferencerDepthLimit() const
 {
-	MaxSearchDepth = NewDepthLimit;
+	return MaxSearchReferencerDepth;
+}
+
+void UEdGraph_ReferenceViewer::SetSearchReferencerDepthLimit(int32 NewDepthLimit)
+{
+	MaxSearchReferencerDepth = NewDepthLimit;
+}
+
+int32 UEdGraph_ReferenceViewer::GetSearchBreadthLimit() const
+{
+	return MaxSearchBreadth;
 }
 
 void UEdGraph_ReferenceViewer::SetSearchBreadthLimit(int32 NewBreadthLimit)
@@ -311,25 +356,39 @@ UEdGraphNode_Reference* UEdGraph_ReferenceViewer::ConstructNodes(const TArray<FA
 		// It both were false, nothing (other than the GraphRootIdentifiers) would be displayed
 		check(bIsShowReferencers || bIsShowDependencies);
 
-		TSet<FName> AllowedPackageNames;
+		// Refresh the current collection filter
+		CurrentCollectionPackages.Empty();
 		if (ShouldFilterByCollection())
 		{
 			FCollectionManagerModule& CollectionManagerModule = FCollectionManagerModule::GetModule();
 			TArray<FName> AssetPaths;
 			CollectionManagerModule.Get().GetAssetsInCollection(CurrentCollectionFilter, ECollectionShareType::CST_All, AssetPaths);
-			AllowedPackageNames.Reserve(AssetPaths.Num());
+			CurrentCollectionPackages.Reserve(AssetPaths.Num());
 			for (FName AssetPath : AssetPaths)
 			{
-				AllowedPackageNames.Add(FName(*FPackageName::ObjectPathToPackageName(AssetPath.ToString())));
+				CurrentCollectionPackages.Add(FName(*FPackageName::ObjectPathToPackageName(AssetPath.ToString())));
 			}
 		}
 
+		// Create & Populate the NodeInfo Maps 
+		// Note to add an empty parent to the root so that if the root node again gets found again as a duplicate, that next parent won't be 
+		// identified as the primary root and also it will appear as having multiple parents.
+		TMap<FAssetIdentifier, ReferenceNodeInfo> ReferenceNodeInfos;
+		ReferenceNodeInfo& RootNodeInfo = ReferenceNodeInfos.FindOrAdd( GraphRootIdentifiers[0], ReferenceNodeInfo(GraphRootIdentifiers[0], true));
+		RootNodeInfo.Parents.Emplace(FAssetIdentifier(NAME_None));
+		RecursivelyPopulateNodeInfos(true, GraphRootIdentifiers[0], ReferenceNodeInfos, 0, MaxSearchReferencerDepth);
+
+		TMap<FAssetIdentifier, ReferenceNodeInfo> DependencyNodeInfos;
+		ReferenceNodeInfo& DRootNodeInfo = DependencyNodeInfos.FindOrAdd( GraphRootIdentifiers[0], ReferenceNodeInfo(GraphRootIdentifiers[0], false));
+		DRootNodeInfo.Parents.Emplace(FAssetIdentifier(NAME_None));
+		RecursivelyPopulateNodeInfos(false, GraphRootIdentifiers[0], DependencyNodeInfos, 0, MaxSearchDependencyDepth);
+
 		TMap<FAssetIdentifier, int32> ReferencerNodeSizes;
 		TSet<FAssetIdentifier> VisitedReferencerSizeNames;
 		if (bIsShowReferencers)
 		{
 			const int32 ReferencerDepth = 1;
-			RecursivelyGatherSizes(/*bReferencers=*/true, GraphRootIdentifiers, AllowedPackageNames, ReferencerDepth, VisitedReferencerSizeNames, ReferencerNodeSizes);
+			RecursivelyGatherSizes(/*bReferencers=*/true, GraphRootIdentifiers, CurrentCollectionPackages, ReferencerDepth, MaxSearchReferencerDepth, VisitedReferencerSizeNames, ReferencerNodeSizes);
 		}
 
 		TMap<FAssetIdentifier, int32> DependencyNodeSizes;
@@ -337,11 +396,10 @@ UEdGraphNode_Reference* UEdGraph_ReferenceViewer::ConstructNodes(const TArray<FA
 		if (bIsShowDependencies)
 		{
 			const int32 DependencyDepth = 1;
-			RecursivelyGatherSizes(/*bReferencers=*/false, GraphRootIdentifiers, AllowedPackageNames, DependencyDepth, VisitedDependencySizeNames, DependencyNodeSizes);
+			RecursivelyGatherSizes(/*bReferencers=*/false, GraphRootIdentifiers, CurrentCollectionPackages, DependencyDepth, MaxSearchDependencyDepth, VisitedDependencySizeNames, DependencyNodeSizes);
 		}
 
 		TSet<FName> AllPackageNames;
-
 		auto AddPackage = [](const FAssetIdentifier& AssetId, TSet<FName>& PackageNames)
 		{ 
 			// Only look for asset data if this is a package
@@ -370,22 +428,49 @@ UEdGraphNode_Reference* UEdGraph_ReferenceViewer::ConstructNodes(const TArray<FA
 		TMap<FName, FAssetData> PackagesToAssetDataMap;
 		GatherAssetData(AllPackageNames, PackagesToAssetDataMap);
 
+		// Store the AssetData in the NodeInfos
+		for (TPair<FAssetIdentifier, ReferenceNodeInfo>&  InfoPair : ReferenceNodeInfos)
+		{
+			InfoPair.Value.AssetData = PackagesToAssetDataMap.FindRef(InfoPair.Key.PackageName);
+		}
+
+		for (TPair<FAssetIdentifier, ReferenceNodeInfo>&  InfoPair : DependencyNodeInfos)
+		{
+			InfoPair.Value.AssetData = PackagesToAssetDataMap.FindRef(InfoPair.Key.PackageName);
+		}
+
 		// Create the root node
 		RootNode = CreateReferenceNode();
-		RootNode->SetupReferenceNode(GraphRootOrigin, GraphRootIdentifiers, PackagesToAssetDataMap.FindRef(GraphRootIdentifiers[0].PackageName), /*bInAllowThumbnail = */ !bIsCompactMode);
+		bool bRootIsDuplicated = DependencyNodeInfos[GraphRootIdentifiers[0]].IsADuplicate() || ReferenceNodeInfos[GraphRootIdentifiers[0]].IsADuplicate();
+		RootNode->SetupReferenceNode(GraphRootOrigin, GraphRootIdentifiers, PackagesToAssetDataMap.FindRef(GraphRootIdentifiers[0].PackageName), /*bInAllowThumbnail = */ !bIsCompactMode, /*bIsDuplicate*/ bRootIsDuplicated);
 
-		if (bIsShowReferencers)
+		if (bUseNodeInfos) // @LULU true = using new stuff, false = Not using new stuff 
 		{
-			TSet<FAssetIdentifier> VisitedReferencerNames;
-			const int32 VisitedReferencerDepth = 1;
-			RecursivelyConstructNodes(/*bReferencers=*/true, RootNode, GraphRootIdentifiers, GraphRootOrigin, ReferencerNodeSizes, PackagesToAssetDataMap, AllowedPackageNames, VisitedReferencerDepth, VisitedReferencerNames);
+			if (bIsShowReferencers)
+			{
+				RecursivelyCreateNodes(true, GraphRootIdentifiers[0], GraphRootOrigin, GraphRootIdentifiers[0], RootNode, ReferenceNodeInfos, 0, MaxSearchReferencerDepth, /*bIsRoot*/ true);
+			}
+
+			if (bIsShowDependencies)
+			{
+				RecursivelyCreateNodes(false, GraphRootIdentifiers[0], GraphRootOrigin, GraphRootIdentifiers[0], RootNode, DependencyNodeInfos, 0, MaxSearchDependencyDepth, /*bIsRoot*/ true);
+			}
 		}
+		else {
 
-		if (bIsShowDependencies)
-		{
-			TSet<FAssetIdentifier> VisitedDependencyNames;
-			const int32 VisitedDependencyDepth = 1;
-			RecursivelyConstructNodes(/*bReferencers=*/false, RootNode, GraphRootIdentifiers, GraphRootOrigin, DependencyNodeSizes, PackagesToAssetDataMap, AllowedPackageNames, VisitedDependencyDepth, VisitedDependencyNames);
+			if (bIsShowReferencers)
+			{
+				TSet<FAssetIdentifier> VisitedReferencerNames;
+				const int32 VisitedReferencerDepth = 1;
+				RecursivelyConstructNodes(/*bReferencers=*/true, RootNode, GraphRootIdentifiers, GraphRootOrigin, ReferencerNodeSizes, PackagesToAssetDataMap, CurrentCollectionPackages, VisitedReferencerDepth, MaxSearchReferencerDepth, VisitedReferencerNames);
+			}
+
+			if (bIsShowDependencies)
+			{
+				TSet<FAssetIdentifier> VisitedDependencyNames;
+				const int32 VisitedDependencyDepth = 1;
+				RecursivelyConstructNodes(/*bReferencers=*/false, RootNode, GraphRootIdentifiers, GraphRootOrigin, DependencyNodeSizes, PackagesToAssetDataMap, CurrentCollectionPackages, VisitedDependencyDepth, MaxSearchDependencyDepth, VisitedDependencyNames);
+			}
 		}
 	}
 
@@ -454,12 +539,28 @@ void UEdGraph_ReferenceViewer::GetSortedLinks(const TArray<FAssetIdentifier>& Id
 		}
 	}
 
+	// Check filters and Filter for our registry source
+	TArray<FAssetIdentifier> ReferenceIds;
+	OutLinks.GenerateKeyArray(ReferenceIds);
+	IAssetManagerEditorModule::Get().FilterAssetIdentifiersForCurrentRegistrySource(ReferenceIds, GetReferenceSearchFlags(false), !bReferencers);
+
 	for (TMap<FAssetIdentifier, EDependencyPinCategory>::TIterator It(OutLinks); It; ++It)
 	{
 		if (!IsPackageIdentifierPassingFilter(It.Key()))
 		{
 			It.RemoveCurrent();
 		}
+
+		else if (!ReferenceIds.Contains(It.Key()))
+		{
+			It.RemoveCurrent();
+		}
+
+		// Collection Filter
+		else if ( ShouldFilterByCollection() && It.Key().IsPackage() && !CurrentCollectionPackages.Contains(It.Key().PackageName))
+		{
+			It.RemoveCurrent();
+		}
 	}
 }
 
@@ -481,7 +582,69 @@ bool UEdGraph_ReferenceViewer::IsPackageIdentifierPassingFilter(const FAssetIden
 	return true;
 }
 
-int32 UEdGraph_ReferenceViewer::RecursivelyGatherSizes(bool bReferencers, const TArray<FAssetIdentifier>& Identifiers, const TSet<FName>& AllowedPackageNames, int32 CurrentDepth, TSet<FAssetIdentifier>& VisitedNames, TMap<FAssetIdentifier, int32>& OutNodeSizes) const
+void
+UEdGraph_ReferenceViewer::RecursivelyPopulateNodeInfos(bool bInReferencers, const FAssetIdentifier& InAssetId, TMap<FAssetIdentifier, ReferenceNodeInfo>& InNodeInfos, int32 InCurrentDepth, int32 InMaxDepth)
+{
+	int32 ProvisionSize = 0;
+
+	int32 Breadth = 0;
+	if (InMaxDepth > 0 && InCurrentDepth < InMaxDepth)
+	{
+		TArray<FAssetIdentifier> AssetIds;
+		AssetIds.Add(InAssetId);
+
+		TMap<FAssetIdentifier, EDependencyPinCategory> ReferenceLinks;
+		GetSortedLinks(AssetIds, bInReferencers, GetReferenceSearchFlags(false), ReferenceLinks);
+
+		InNodeInfos[InAssetId].Children.Reserve(ReferenceLinks.Num());
+		for (const TPair<FAssetIdentifier, EDependencyPinCategory>& Pair : ReferenceLinks)
+		{
+			FAssetIdentifier ChildId = Pair.Key;
+			if (!ExceedsMaxSearchBreadth(Breadth))// Breadth < MaxSearchBreadth)
+			{
+				// Only Gather Children the first time
+				if (!InNodeInfos.Contains(ChildId))
+				{
+					ReferenceNodeInfo& NewNodeInfo = InNodeInfos.FindOrAdd(ChildId, ReferenceNodeInfo(ChildId, bInReferencers));
+					InNodeInfos[ChildId].Parents.Emplace(InAssetId);
+					InNodeInfos[InAssetId].Children.Emplace(Pair);
+
+					RecursivelyPopulateNodeInfos(bInReferencers, ChildId, InNodeInfos, InCurrentDepth + 1, InMaxDepth);
+					ProvisionSize += InNodeInfos[ChildId].ProvisionSize(InAssetId);
+					Breadth++;
+				}
+
+				else if (bIsShowDuplicates && !InNodeInfos[ChildId].Parents.Contains(InAssetId))
+				{
+					InNodeInfos[ChildId].Parents.Emplace(InAssetId);
+					InNodeInfos[InAssetId].Children.Emplace(Pair);
+					ProvisionSize += 1;
+					Breadth++;
+				}
+			}
+			else
+			{
+				// Count the overflow nodes to report in the UI but otherwise skip adding them 
+				if (bIsShowDuplicates || !InNodeInfos.Contains(ChildId))
+				{
+					InNodeInfos[InAssetId].OverflowCount++;
+					Breadth++;
+				}
+			}
+
+		}
+	}
+
+	// Account for an overflow node if necessary
+	if (InNodeInfos[InAssetId].OverflowCount > 0)
+	{
+		ProvisionSize++;
+	}
+
+	InNodeInfos[InAssetId].ChildProvisionSize = ProvisionSize > 0 ? ProvisionSize : 1;
+}
+
+int32 UEdGraph_ReferenceViewer::RecursivelyGatherSizes(bool bReferencers, const TArray<FAssetIdentifier>& Identifiers, const TSet<FName>& AllowedPackageNames, int32 CurrentDepth, int32 MaxDepth, TSet<FAssetIdentifier>& VisitedNames, TMap<FAssetIdentifier, int32>& OutNodeSizes) const
 {
 	check(Identifiers.Num() > 0);
 
@@ -490,33 +653,23 @@ int32 UEdGraph_ReferenceViewer::RecursivelyGatherSizes(bool bReferencers, const
 	TMap<FAssetIdentifier, EDependencyPinCategory> ReferenceLinks;
 	GetSortedLinks(Identifiers, bReferencers, GetReferenceSearchFlags(false), ReferenceLinks);
 
-	TArray<FAssetIdentifier> ReferenceNames;
-	ReferenceNames.Reserve(ReferenceLinks.Num());
-	for (const TPair<FAssetIdentifier, EDependencyPinCategory>& Pair : ReferenceLinks)
-	{
-		ReferenceNames.Add(Pair.Key);
-	}
-
 	int32 NodeSize = 0;
-	if ( ReferenceNames.Num() > 0 && !ExceedsMaxSearchDepth(CurrentDepth) )
+
+	if ( ReferenceLinks.Num() > 0 && !ExceedsMaxSearchDepth(CurrentDepth, MaxDepth) )
 	{
 		int32 NumReferencesMade = 0;
 		int32 NumReferencesExceedingMax = 0;
 
-		// Filter for our registry source
-		IAssetManagerEditorModule::Get().FilterAssetIdentifiersForCurrentRegistrySource(ReferenceNames, GetReferenceSearchFlags(false), !bReferencers);
-
-		// Since there are referencers, use the size of all your combined referencers.
-		// Do not count your own size since there could just be a horizontal line of nodes
-		for (FAssetIdentifier& AssetId : ReferenceNames)
+		for (const TPair<FAssetIdentifier, EDependencyPinCategory>& Pair : ReferenceLinks)
 		{
-			if ( !VisitedNames.Contains(AssetId) && (!AssetId.IsPackage() || !ShouldFilterByCollection() || AllowedPackageNames.Contains(AssetId.PackageName)) )
+			const FAssetIdentifier& AssetId = Pair.Key;
+			if ( !VisitedNames.Contains(AssetId) )
 			{
 				if ( !ExceedsMaxSearchBreadth(NumReferencesMade) )
 				{
 					TArray<FAssetIdentifier> NewPackageNames;
 					NewPackageNames.Add(AssetId);
-					NodeSize += RecursivelyGatherSizes(bReferencers, NewPackageNames, AllowedPackageNames, CurrentDepth + 1, VisitedNames, OutNodeSizes);
+					NodeSize += RecursivelyGatherSizes(bReferencers, NewPackageNames, AllowedPackageNames, CurrentDepth + 1, MaxDepth, VisitedNames, OutNodeSizes);
 					NumReferencesMade++;
 				}
 				else
@@ -548,7 +701,94 @@ void UEdGraph_ReferenceViewer::GatherAssetData(const TSet<FName>& AllPackageName
 	UE::AssetRegistry::GetAssetForPackages(AllPackageNames.Array(), OutPackageToAssetDataMap);
 }
 
-UEdGraphNode_Reference* UEdGraph_ReferenceViewer::RecursivelyConstructNodes(bool bReferencers, UEdGraphNode_Reference* RootNode, const TArray<FAssetIdentifier>& Identifiers, const FIntPoint& NodeLoc, const TMap<FAssetIdentifier, int32>& NodeSizes, const TMap<FName, FAssetData>& PackagesToAssetDataMap, const TSet<FName>& AllowedPackageNames, int32 CurrentDepth, TSet<FAssetIdentifier>& VisitedNames)
+UEdGraphNode_Reference* UEdGraph_ReferenceViewer::RecursivelyCreateNodes(bool bInReferencers, const FAssetIdentifier& InAssetId, const FIntPoint& InNodeLoc, const FAssetIdentifier& InParentId, UEdGraphNode_Reference* InParentNode, TMap<FAssetIdentifier, ReferenceNodeInfo>& InNodeInfos, int32 InCurrentDepth, int32 InMaxDepth, bool bIsRoot)
+{
+	check(InNodeInfos.Contains(InAssetId));
+
+	const ReferenceNodeInfo& NodeInfo = InNodeInfos[InAssetId];
+	int32 NodeProvSize = 1;
+
+	UEdGraphNode_Reference* NewNode = NULL;
+	if (bIsRoot)
+	{
+		NewNode = InParentNode;
+		NodeProvSize = NodeInfo.ProvisionSize(FAssetIdentifier(NAME_None));
+	}
+	else
+	{
+		NewNode = CreateReferenceNode();
+		NewNode->SetupReferenceNode(InNodeLoc, {InAssetId}, NodeInfo.AssetData, /*bInAllowThumbnail*/ !bIsCompactMode, /*bIsADuplicate*/ NodeInfo.Parents.Num() > 1);
+		NodeProvSize = NodeInfo.ProvisionSize(InParentId);
+	}
+
+	bool bIsFirstOccurance = bIsRoot || NodeInfo.IsFirstParent(InParentId);
+	FIntPoint ChildLoc = InNodeLoc;
+	if (InMaxDepth > 0 && (InCurrentDepth < InMaxDepth) && bIsFirstOccurance) // Only expand the first parent
+	{
+
+		// position the children nodes
+		const int32 ColumnWidth = bIsCompactMode ? 400 : 800;
+		ChildLoc.X += bInReferencers ? -ColumnWidth : ColumnWidth;
+
+		const int32 NodeSizeY = bIsCompactMode ? 100 : 200;
+		ChildLoc.Y -= (NodeProvSize - 1) * NodeSizeY * 0.5 ;
+
+		int32 Breadth = 0;
+
+		for (const TPair<FAssetIdentifier, EDependencyPinCategory>& Pair : InNodeInfos[InAssetId].Children)
+		{
+
+			    FAssetIdentifier ChildId = Pair.Key;
+			    int32 ChildProvSize = InNodeInfos[ChildId].ProvisionSize(InAssetId);
+
+				ChildLoc.Y += (ChildProvSize - 1) * NodeSizeY * 0.5;
+
+				UEdGraphNode_Reference* ChildNode = RecursivelyCreateNodes(bInReferencers, ChildId, ChildLoc, InAssetId, NewNode, InNodeInfos, InCurrentDepth + 1, InMaxDepth);	
+
+				if (bInReferencers)
+				{
+					ChildNode->GetDependencyPin()->PinType.PinCategory = ::GetName(Pair.Value);
+					NewNode->AddReferencer( ChildNode );
+				}
+				else
+				{
+					ChildNode->GetReferencerPin()->PinType.PinCategory = ::GetName(Pair.Value);
+					ChildNode->AddReferencer( NewNode );
+				}
+
+				ChildLoc.Y += NodeSizeY * (ChildProvSize + 1) * 0.5;
+
+		}
+
+		// There were more references than allowed to be displayed. Make a collapsed node.
+		if (NodeInfo.OverflowCount > 0)
+		{
+			UEdGraphNode_Reference* OverflowNode = CreateReferenceNode();
+			FIntPoint RefNodeLoc;
+			RefNodeLoc.X = ChildLoc.X;
+			RefNodeLoc.Y = ChildLoc.Y;
+
+			if ( ensure(OverflowNode) )
+			{
+				OverflowNode->SetAllowThumbnail(!bIsCompactMode);
+				OverflowNode->SetReferenceNodeCollapsed(RefNodeLoc, NodeInfo.OverflowCount);
+
+				if ( bInReferencers )
+				{
+					NewNode->AddReferencer( OverflowNode );
+				}
+				else
+				{
+					OverflowNode->AddReferencer( NewNode );
+				}
+			}
+		}
+	}
+
+	return NewNode;
+}
+
+UEdGraphNode_Reference* UEdGraph_ReferenceViewer::RecursivelyConstructNodes(bool bReferencers, UEdGraphNode_Reference* RootNode, const TArray<FAssetIdentifier>& Identifiers, const FIntPoint& NodeLoc, const TMap<FAssetIdentifier, int32>& NodeSizes, const TMap<FName, FAssetData>& PackagesToAssetDataMap, const TSet<FName>& AllowedPackageNames, int32 CurrentDepth, int32 MaxDepth, TSet<FAssetIdentifier>& VisitedNames)
 {
 	check(Identifiers.Num() > 0);
 
@@ -563,13 +803,13 @@ UEdGraphNode_Reference* UEdGraph_ReferenceViewer::RecursivelyConstructNodes(bool
 	else
 	{
 		NewNode = CreateReferenceNode();
-		NewNode->SetupReferenceNode(NodeLoc, Identifiers, PackagesToAssetDataMap.FindRef(Identifiers[0].PackageName), /*bInAllowThumbnail = */ !bIsCompactMode);
+		NewNode->SetupReferenceNode(NodeLoc, Identifiers, PackagesToAssetDataMap.FindRef(Identifiers[0].PackageName), /*bInAllowThumbnail = */ !bIsCompactMode, /*bIsDuplicate*/ false);
 	}
 
 	TMap<FAssetIdentifier, EDependencyPinCategory> Referencers;
 	GetSortedLinks(Identifiers, bReferencers, GetReferenceSearchFlags(false), Referencers);
 
-	if (Referencers.Num() > 0 && !ExceedsMaxSearchDepth(CurrentDepth))
+	if (Referencers.Num() > 0 && !ExceedsMaxSearchDepth(CurrentDepth, MaxDepth))
 	{
 		FIntPoint ReferenceNodeLoc = NodeLoc;
 
@@ -593,31 +833,11 @@ UEdGraphNode_Reference* UEdGraph_ReferenceViewer::RecursivelyConstructNodes(bool
 		int32 NumReferencesMade = 0;
 		int32 NumReferencesExceedingMax = 0;
 
-		// Filter for our registry source
-		TArray<FAssetIdentifier> ReferenceIds;
-		Referencers.GenerateKeyArray(ReferenceIds);
-		IAssetManagerEditorModule::Get().FilterAssetIdentifiersForCurrentRegistrySource(ReferenceIds, GetReferenceSearchFlags(false), !bReferencers);
-		if (ReferenceIds.Num() != Referencers.Num())
-		{
-			TSet<FAssetIdentifier> KeptSet;
-			KeptSet.Reserve(ReferenceIds.Num());
-			for (FAssetIdentifier Id : ReferenceIds)
-			{
-				KeptSet.Add(Id);
-			}
-			for (TMap<FAssetIdentifier, EDependencyPinCategory>::TIterator It(Referencers); It; ++It)
-			{
-				if (!KeptSet.Contains(It.Key()))
-				{
-					It.RemoveCurrent();
-				}
-			}
-		}
 		for ( TPair<FAssetIdentifier, EDependencyPinCategory>& DependencyPair : Referencers)
 		{
 			FAssetIdentifier ReferenceName = DependencyPair.Key;
 
-			if ( !VisitedNames.Contains(ReferenceName) && (!ReferenceName.IsPackage() || !ShouldFilterByCollection() || AllowedPackageNames.Contains(ReferenceName.PackageName)) )
+			if ( !VisitedNames.Contains(ReferenceName) )
 			{
 				if ( !ExceedsMaxSearchBreadth(NumReferencesMade) )
 				{
@@ -631,7 +851,7 @@ UEdGraphNode_Reference* UEdGraph_ReferenceViewer::RecursivelyConstructNodes(bool
 					TArray<FAssetIdentifier> NewIdentifiers;
 					NewIdentifiers.Add(ReferenceName);
 					
-					UEdGraphNode_Reference* ReferenceNode = RecursivelyConstructNodes(bReferencers, RootNode, NewIdentifiers, RefNodeLoc, NodeSizes, PackagesToAssetDataMap, AllowedPackageNames, CurrentDepth + 1, VisitedNames);
+					UEdGraphNode_Reference* ReferenceNode = RecursivelyConstructNodes(bReferencers, RootNode, NewIdentifiers, RefNodeLoc, NodeSizes, PackagesToAssetDataMap, AllowedPackageNames, CurrentDepth + 1, MaxDepth, VisitedNames);
 					if ( ensure(ReferenceNode) )
 					{
 						if (bReferencers)
@@ -697,10 +917,10 @@ const TSharedPtr<FAssetThumbnailPool>& UEdGraph_ReferenceViewer::GetAssetThumbna
 	return AssetThumbnailPool;
 }
 
-bool UEdGraph_ReferenceViewer::ExceedsMaxSearchDepth(int32 Depth) const
+bool UEdGraph_ReferenceViewer::ExceedsMaxSearchDepth(int32 Depth, int32 MaxDepth) const
 {
 	// ExceedsMaxSearchDepth requires only greater (not equal than) because, even though the Depth is 1-based indexed (similarly to Breadth), the first index (index 0) corresponds to the root object 
-	return bLimitSearchDepth && Depth > MaxSearchDepth;
+	return bLimitSearchDepth && Depth > MaxDepth;
 }
 
 bool UEdGraph_ReferenceViewer::ExceedsMaxSearchBreadth(int32 Breadth) const
diff --git a/Engine/Plugins/Editor/AssetManagerEditor/Source/AssetManagerEditor/Private/ReferenceViewer/ReferenceViewerStyle.cpp b/Engine/Plugins/Editor/AssetManagerEditor/Source/AssetManagerEditor/Private/ReferenceViewer/ReferenceViewerStyle.cpp
new file mode 100644
--- /dev/null
+++ b/Engine/Plugins/Editor/AssetManagerEditor/Source/AssetManagerEditor/Private/ReferenceViewer/ReferenceViewerStyle.cpp
@@ -0,0 +1,91 @@
+// Copyright Epic Games, Inc. All Rights Reserved.
+
+#include "ReferenceViewerStyle.h"
+
+#include "Styling/SlateTypes.h"
+#include "Styling/StarshipCoreStyle.h"
+#include "Styling/SlateStyleMacros.h"
+#include "Styling/StyleColors.h"
+
+FName FReferenceViewerStyle::StyleName("ReferenceViewerStyle");
+
+#define FONT(...) FSlateFontInfo(FCoreStyle::GetDefaultFont(), __VA_ARGS__)
+
+FReferenceViewerStyle::FReferenceViewerStyle()
+	: FSlateStyleSet(StyleName)
+{
+	SetParentStyleName("EditorStyle");
+
+	SetContentRoot(FPaths::EnginePluginsDir() / TEXT("Editor/AssetManagerEditor/Content"));
+	SetCoreContentRoot(FPaths::EngineContentDir() / TEXT("Slate"));
+
+	FScrollBarStyle ScrollBar = GetParentStyle()->GetWidgetStyle<FScrollBarStyle>("ScrollBar");
+	FTextBlockStyle NormalText = GetParentStyle()->GetWidgetStyle<FTextBlockStyle>("NormalText");
+	FEditableTextBoxStyle NormalEditableTextBoxStyle = GetParentStyle()->GetWidgetStyle<FEditableTextBoxStyle>("NormalEditableTextBox");
+
+	FSlateFontInfo TitleFont = FONT(12, "Bold");
+
+	// Text Styles 
+	FTextBlockStyle GraphNodeTitle = FTextBlockStyle(NormalText)
+		.SetFont(TitleFont)
+		.SetColorAndOpacity( FStyleColors::White)
+		.SetShadowOffset( FVector2D::UnitVector )
+		.SetShadowColorAndOpacity( FLinearColor::Black );
+	Set( "Graph.Node.NodeTitle", GraphNodeTitle );
+
+
+	Set( "Graph.Node.NodeTitleExtraLines", FTextBlockStyle(NormalText)
+		.SetFont( DEFAULT_FONT( "Normal", 9 ) )
+		.SetColorAndOpacity( FStyleColors::White)
+		.SetShadowOffset( FVector2D::ZeroVector)
+		.SetShadowColorAndOpacity( FLinearColor::Transparent)
+	);
+
+
+	FEditableTextBoxStyle GraphNodeTitleEditableText = FEditableTextBoxStyle(NormalEditableTextBoxStyle)
+		.SetFont(TitleFont)
+		.SetForegroundColor(FStyleColors::Input)
+		.SetBackgroundImageNormal(FSlateRoundedBoxBrush(FStyleColors::Foreground, FStyleColors::Secondary, 1.0f))
+		.SetBackgroundImageHovered(FSlateRoundedBoxBrush(FStyleColors::Foreground, FStyleColors::Hover, 1.0f))
+		.SetBackgroundImageFocused(FSlateRoundedBoxBrush(FStyleColors::Foreground, FStyleColors::Primary, 1.0f))
+		.SetBackgroundImageReadOnly(FSlateRoundedBoxBrush(FStyleColors::Header, FStyleColors::InputOutline, 1.0f))
+		.SetForegroundColor(FStyleColors::White)
+		.SetBackgroundColor(FStyleColors::White)
+		.SetReadOnlyForegroundColor(FStyleColors::Foreground)
+		.SetFocusedForegroundColor(FStyleColors::Background)
+		.SetScrollBarStyle( ScrollBar );
+	Set( "Graph.Node.NodeTitleEditableText", GraphNodeTitleEditableText );
+
+	Set( "Graph.Node.NodeTitleInlineEditableText", FInlineEditableTextBlockStyle()
+		.SetTextStyle(GraphNodeTitle)
+		.SetEditableTextBoxStyle(GraphNodeTitleEditableText)
+	);
+
+	
+	FLinearColor SpillColor(.3, .3, .3, 1.0);
+	int BodyRadius = 10.0; // Designed for 4 but using 10 to accomodate the shared selection border.  Update to 4 all nodes get aligned.
+
+	// NOTE: 
+	Set( "Graph.Node.BodyBackground", new FSlateRoundedBoxBrush(FStyleColors::Panel, BodyRadius));
+	Set( "Graph.Node.BodyBorder",     new FSlateRoundedBoxBrush(SpillColor, BodyRadius));
+
+	Set( "Graph.Node.Body",           new FSlateRoundedBoxBrush(FStyleColors::Panel, BodyRadius, FStyleColors::Transparent, 2.0));
+	Set( "Graph.Node.ColorSpill",     new FSlateRoundedBoxBrush(SpillColor, FVector4(BodyRadius, BodyRadius, 0.0, 0.0)));
+
+	Set( "Graph.Node.Duplicate",      new IMAGE_BRUSH_SVG("/GraphNode_Duplicate_8px", FVector2D(8.f, 8.f), FStyleColors::White));
+	
+	FSlateStyleRegistry::RegisterSlateStyle(*this);
+}
+
+FReferenceViewerStyle::~FReferenceViewerStyle()
+{
+	FSlateStyleRegistry::UnRegisterSlateStyle(*this);
+}
+
+FReferenceViewerStyle& FReferenceViewerStyle::Get()
+{
+	static FReferenceViewerStyle Inst;
+	return Inst;
+}
+
+
diff --git a/Engine/Plugins/Editor/AssetManagerEditor/Source/AssetManagerEditor/Private/ReferenceViewer/ReferenceViewerStyle.h b/Engine/Plugins/Editor/AssetManagerEditor/Source/AssetManagerEditor/Private/ReferenceViewer/ReferenceViewerStyle.h
new file mode 100644
--- /dev/null
+++ b/Engine/Plugins/Editor/AssetManagerEditor/Source/AssetManagerEditor/Private/ReferenceViewer/ReferenceViewerStyle.h
@@ -0,0 +1,30 @@
+// Copyright Epic Games, Inc. All Rights Reserved.
+
+#pragma once
+
+#include "Styling/SlateStyle.h"
+#include "Styling/SlateStyleRegistry.h"
+
+/**
+ * The Reference Viewer is being used as a sandbox for testing out ui changes to graphs.
+ * This style is setup to facilitate testing new styles for graphs and should be deleted
+ * once the main graph style is updated, which currently resides in StarshipStyle (the Editor Style).
+ * 
+ * This style inherits from the EditorStyle (StarshipStyle) and therefore is just overwriting 
+ * styles on the editor style.  No style names should have to be replaced when removing this style. 
+ * However, someone will have to update the references from FReferenceViewerStyle::Get() to be FAppStyle::Get().
+ */
+class FReferenceViewerStyle
+	: public FSlateStyleSet
+{
+public:
+	static FName StyleName;
+
+	/** Access the singleton instance for this style set */
+	static FReferenceViewerStyle& Get();
+
+private:
+
+	FReferenceViewerStyle();
+	~FReferenceViewerStyle();
+};
diff --git a/Engine/Plugins/Editor/AssetManagerEditor/Source/AssetManagerEditor/Private/ReferenceViewer/SReferenceNode.cpp b/Engine/Plugins/Editor/AssetManagerEditor/Source/AssetManagerEditor/Private/ReferenceViewer/SReferenceNode.cpp
--- a/Engine/Plugins/Editor/AssetManagerEditor/Source/AssetManagerEditor/Private/ReferenceViewer/SReferenceNode.cpp
+++ b/Engine/Plugins/Editor/AssetManagerEditor/Source/AssetManagerEditor/Private/ReferenceViewer/SReferenceNode.cpp
@@ -9,6 +9,7 @@
 #include "AssetThumbnail.h"
 #include "Widgets/Text/SInlineEditableTextBlock.h"
 #include "SCommentBubble.h"
+#include "ReferenceViewerStyle.h"
 
 #define LOCTEXT_NAMESPACE "ReferenceViewer"
 
@@ -36,6 +37,7 @@ void SReferenceNode::Construct( const FArguments& InArgs, UEdGraphNode_Reference
 	UpdateGraphNode();
 }
 
+// UpdateGraphNode is similar to the base, but adds the option to hide the thumbnail */
 void SReferenceNode::UpdateGraphNode()
 {
 	OutputPins.Empty();
@@ -59,12 +61,22 @@ void SReferenceNode::UpdateGraphNode()
 	//
 	TSharedPtr<SVerticalBox> MainVerticalBox;
 	TSharedPtr<SErrorText> ErrorText;
-	TSharedPtr<SNodeTitle> NodeTitle = SNew(SNodeTitle, GraphNode);
+	TSharedPtr<SNodeTitle> NodeTitle = SNew(SNodeTitle, GraphNode).StyleSet(&FReferenceViewerStyle::Get());
+
+	// Get node icon
+	IconColor = FLinearColor::White;
+	const FSlateBrush* IconBrush = nullptr;
+	if (GraphNode != NULL && GraphNode->ShowPaletteIconOnNode())
+	{
+		IconBrush = GraphNode->GetIconAndTint(IconColor).GetOptionalIcon();
+	}
+
 
 	TSharedRef<SWidget> ThumbnailWidget = SNullWidget::NullWidget;
+	UEdGraphNode_Reference* RefGraphNode = CastChecked<UEdGraphNode_Reference>(GraphNode);
+	bool bIsADuplicate = RefGraphNode->IsADuplicate();
 	if ( AssetThumbnail.IsValid() )
 	{
-		UEdGraphNode_Reference* RefGraphNode = CastChecked<UEdGraphNode_Reference>(GraphNode);
 
 		FAssetThumbnailConfig ThumbnailConfig;
 		ThumbnailConfig.bAllowFadeIn = RefGraphNode->UsesThumbnail();
@@ -89,39 +101,56 @@ void SReferenceNode::UpdateGraphNode()
 		.AutoHeight()
 		[
 			SNew(SBorder)
-			.BorderImage( FEditorStyle::GetBrush( "Graph.Node.Body" ) )
+			.BorderImage( FReferenceViewerStyle::Get().GetBrush( "Graph.Node.BodyBackground" ) )
+			.Padding(0)
+			[
+
+			SNew(SBorder)
+			.BorderImage( FReferenceViewerStyle::Get().GetBrush( "Graph.Node.BodyBorder" ) )
+			.BorderBackgroundColor( this, &SReferenceNode::GetNodeTitleColor )
 			.Padding(0)
 			[
+
+			SNew(SBorder)
+			.BorderImage( FReferenceViewerStyle::Get().GetBrush( "Graph.Node.Body" ) )
+			.Padding(0)
+			[
+
 				SNew(SVerticalBox)
-				. ToolTipText( this, &SReferenceNode::GetNodeTooltip )
+				.ToolTipText( this, &SReferenceNode::GetNodeTooltip )
+
 				+SVerticalBox::Slot()
 				.AutoHeight()
 				.HAlign(HAlign_Fill)
 				.VAlign(VAlign_Top)
+				.Padding(0)
 				[
-					SNew(SOverlay)
-					+SOverlay::Slot()
-					[
-						SNew(SImage)
-						.Image( FEditorStyle::GetBrush("Graph.Node.TitleGloss") )
-					]
-					+SOverlay::Slot()
-					.HAlign(HAlign_Left)
-					.VAlign(VAlign_Center)
+					SNew(SBorder)
+					.BorderImage( FReferenceViewerStyle::Get().GetBrush("Graph.Node.ColorSpill") )
+					.Padding( FMargin(10.0f, 4.0f, 6.0f, 4.0f) )
+					.BorderBackgroundColor( this, &SReferenceNode::GetNodeTitleColor )
 					[
-						SNew(SBorder)
-						.BorderImage( FEditorStyle::GetBrush("Graph.Node.ColorSpill") )
-						// The extra margin on the right
-						// is for making the color spill stretch well past the node title
-						.Padding( FMargin(10,5,30,3) )
-						.BorderBackgroundColor( this, &SReferenceNode::GetNodeTitleColor )
+						SNew(SHorizontalBox)
+						+ SHorizontalBox::Slot()
+						.VAlign(VAlign_Center)
+						.Padding(FMargin(0.f, 0.f, 8.f, 0.f))
+						.AutoWidth()
+						[
+							SNew(SImage)
+							.Image(IconBrush)
+							.DesiredSizeOverride(FVector2D(24.0, 24.0))
+							.ColorAndOpacity(this, &SGraphNode::GetNodeTitleIconColor)
+						]
+						+ SHorizontalBox::Slot()
+						.FillWidth(1.0)
 						[
 							SNew(SVerticalBox)
 							+SVerticalBox::Slot()
-								.AutoHeight()
+							.AutoHeight()
+							.Padding(FMargin(0.f))
 							[
 								SAssignNew(InlineEditableText, SInlineEditableTextBlock)
-								.Style( FEditorStyle::Get(), "Graph.Node.NodeTitleInlineEditableText" )
+								.Style( FReferenceViewerStyle::Get(), "Graph.Node.NodeTitleInlineEditableText" )
 								.Text( NodeTitle.Get(), &SNodeTitle::GetHeadTitle )
 								.OnVerifyTextChanged(this, &SReferenceNode::OnVerifyNameTextChanged)
 								.OnTextCommitted(this, &SReferenceNode::OnNameTextCommited)
@@ -129,21 +158,25 @@ void SReferenceNode::UpdateGraphNode()
 								.IsSelected(this, &SReferenceNode::IsSelectedExclusively)
 							]
 							+SVerticalBox::Slot()
-								.AutoHeight()
+							.AutoHeight()
+							.Padding(FMargin(0.f))
 							[
 								NodeTitle.ToSharedRef()
 							]
 						]
-					]
-					+SOverlay::Slot()
-					.VAlign(VAlign_Top)
-					[
-						SNew(SBorder)
-						.BorderImage( FEditorStyle::GetBrush( "Graph.Node.TitleHighlight" ) )
-						.Visibility(EVisibility::HitTestInvisible)			
+
+						+SHorizontalBox::Slot()
+						.AutoWidth()
+						.HAlign(HAlign_Right)
+						.VAlign(VAlign_Bottom)
+						.Padding(FMargin(0.f, 0.f, 0.f, 1.f))
 						[
-							SNew(SSpacer)
-							.Size(FVector2D(20,20))
+							SNew(SImage)
+							.Visibility(bIsADuplicate ? EVisibility::Visible : EVisibility::Hidden)
+							.ToolTipText(LOCTEXT("DuplicateAsset", "This asset is referenced multiple times. Only the first occurance shows its decendants."))
+							.DesiredSizeOverride(FVector2D(12.0, 12.0))
+							.Image(FReferenceViewerStyle::Get().GetBrush("Graph.Node.Duplicate"))
+							.ColorAndOpacity(FAppStyle::Get().GetColor("Colors.Foreground"))
 						]
 					]
 				]
@@ -203,7 +236,10 @@ void SReferenceNode::UpdateGraphNode()
 						]
 					]
 				]
-			]
+
+			] // Inner Background 
+			] // Outline Border
+			] // Background 
 		]
 	];
 	// Create comment bubble if comment text is valid
diff --git a/Engine/Plugins/Editor/AssetManagerEditor/Source/AssetManagerEditor/Private/ReferenceViewer/SReferenceViewer.cpp b/Engine/Plugins/Editor/AssetManagerEditor/Source/AssetManagerEditor/Private/ReferenceViewer/SReferenceViewer.cpp
--- a/Engine/Plugins/Editor/AssetManagerEditor/Source/AssetManagerEditor/Private/ReferenceViewer/SReferenceViewer.cpp
+++ b/Engine/Plugins/Editor/AssetManagerEditor/Source/AssetManagerEditor/Private/ReferenceViewer/SReferenceViewer.cpp
@@ -39,6 +39,12 @@
 
 #define LOCTEXT_NAMESPACE "ReferenceViewer"
 
+static bool GShowToggleDeprecatedReferenceViewerLayout = false;
+static FAutoConsoleCommand CVarShowToggleDeprecatedReferenceViewerLayout(
+	TEXT("ReferenceViewer.ShowToggleDeprecatedLayout"),
+	TEXT("Displays the toggle allowing the user to toggle back to the former layout algorithm."),
+	FConsoleCommandDelegate::CreateLambda([] { GShowToggleDeprecatedReferenceViewerLayout = !GShowToggleDeprecatedReferenceViewerLayout; }));
+
 bool IsPackageNamePassingFilter(FName InPackageName, const TArray<FString>& InSearchWords)
 {
 	// package name must match all words
@@ -281,7 +287,7 @@ void SReferenceViewer::Construct(const FArguments& InArgs)
 						.Padding(2.f)
 						[
 							SNew(STextBlock)
-							.Text(LOCTEXT("SearchDepthLabelText", "Search Depth Limit"))
+							.Text(LOCTEXT("SearchDepthReferencersLabelText", "Search Referencers Depth"))
 						]
 
 						+SHorizontalBox::Slot()
@@ -289,11 +295,33 @@ void SReferenceViewer::Construct(const FArguments& InArgs)
 						.VAlign(VAlign_Center)
 						.Padding(2.f)
 						[
-							SNew(SCheckBox)
-							.OnCheckStateChanged( this, &SReferenceViewer::OnSearchDepthEnabledChanged )
-							.IsChecked( this, &SReferenceViewer::IsSearchDepthEnabledChecked )
+							SNew(SBox)
+							.WidthOverride(100)
+							[
+								SNew(SSpinBox<int32>)
+								.Value(this, &SReferenceViewer::GetSearchReferencerDepthCount)
+								.OnValueChanged(this, &SReferenceViewer::OnSearchReferencerDepthCommitted)
+								.MinValue(0)
+								.MaxValue(50)
+								.MaxSliderValue(12)
+							]
 						]
-					
+					]
+
+					+SVerticalBox::Slot()
+					.AutoHeight()
+					[
+						SNew(SHorizontalBox)
+						.Visibility_Lambda([this]() { return (FixAndHideSearchDepthLimit > 0 ? EVisibility::Collapsed : EVisibility::Visible); })
+
+						+SHorizontalBox::Slot()
+						.VAlign(VAlign_Center)
+						.Padding(2.f)
+						[
+							SNew(STextBlock)
+							.Text(LOCTEXT("SearchDepthDependenciesLabelText", "Search Dependencies Depth"))
+						]
+
 						+SHorizontalBox::Slot()
 						.AutoWidth()
 						.VAlign(VAlign_Center)
@@ -303,9 +331,9 @@ void SReferenceViewer::Construct(const FArguments& InArgs)
 							.WidthOverride(100)
 							[
 								SNew(SSpinBox<int32>)
-								.Value(this, &SReferenceViewer::GetSearchDepthCount)
-								.OnValueChanged(this, &SReferenceViewer::OnSearchDepthCommitted)
-								.MinValue(1)
+								.Value(this, &SReferenceViewer::GetSearchDependencyDepthCount)
+								.OnValueChanged(this, &SReferenceViewer::OnSearchDependencyDepthCommitted)
+								.MinValue(0)
 								.MaxValue(50)
 								.MaxSliderValue(12)
 							]
@@ -575,6 +603,63 @@ void SReferenceViewer::Construct(const FArguments& InArgs)
 							.IsChecked(this, &SReferenceViewer::IsCompactModeChecked)
 						]
 					]
+
+					+ SVerticalBox::Slot()
+					.AutoHeight()
+					[
+						SNew(SHorizontalBox)
+						+ SHorizontalBox::Slot()
+						.VAlign(VAlign_Center)
+						.Padding(2.f)
+						[
+							SNew(STextBlock)
+							.Text(LOCTEXT("DuplicateAssets", "Show Duplicate Assets"))
+						]
+
+						+SHorizontalBox::Slot()
+						.AutoWidth()
+						.VAlign(VAlign_Center)
+						.Padding(2.f)
+						[
+							SNew(SCheckBox)
+							.OnCheckStateChanged(this, &SReferenceViewer::OnShowDuplicatesChanged)
+							.IsChecked(this, &SReferenceViewer::IsShowDuplicatesChecked)
+						]
+					]
+
+					+ SVerticalBox::Slot()
+					.AutoHeight()
+					[
+						SNew(SHorizontalBox)
+						.Visibility_Lambda([]() { return (GShowToggleDeprecatedReferenceViewerLayout ? EVisibility::Visible : EVisibility::Collapsed); })
+
+						+ SHorizontalBox::Slot()
+						.VAlign(VAlign_Center)
+						.Padding(2.f)
+						[
+							SNew(STextBlock)
+							.Text(LOCTEXT("UseOldLayoutMechanism", "Use Deprecated Layout"))
+						]
+
+						+SHorizontalBox::Slot()
+						.AutoWidth()
+						.VAlign(VAlign_Center)
+						.Padding(2.f)
+						[
+							SNew(SCheckBox)
+							.OnCheckStateChanged_Lambda( [this] (ECheckBoxState NewState) { 
+								if (GraphObj) 
+								{ 	
+									GraphObj->SetUseNodeInfos(NewState != ECheckBoxState::Checked); 
+									RebuildGraph();} 
+								} 
+							)
+							.IsChecked_Lambda([this] () 
+							{
+								return (GraphObj && GraphObj->GetUseNodeInfos()) ? ECheckBoxState::Unchecked: ECheckBoxState::Checked;
+							})
+						]
+					]
 				]
 			]
 
@@ -610,7 +695,8 @@ void SReferenceViewer::SetGraphRootIdentifiers(const TArray<FAssetIdentifier>& N
 	FixAndHideSearchDepthLimit = ReferenceViewerParams.FixAndHideSearchDepthLimit;
 	if (FixAndHideSearchDepthLimit > 0)
 	{
-		GraphObj->SetSearchDepthLimit(FixAndHideSearchDepthLimit);
+		GraphObj->SetSearchDependencyDepthLimit(FixAndHideSearchDepthLimit);
+		GraphObj->SetSearchReferencerDepthLimit(FixAndHideSearchDepthLimit);
 		GraphObj->SetSearchDepthLimitEnabled(true);
 	}
 	FixAndHideSearchBreadthLimit = ReferenceViewerParams.FixAndHideSearchBreadthLimit;
@@ -895,11 +981,11 @@ ECheckBoxState SReferenceViewer::IsSearchDepthEnabledChecked() const
 	}
 }
 
-int32 SReferenceViewer::GetSearchDepthCount() const
+int32 SReferenceViewer::GetSearchDependencyDepthCount() const
 {
 	if ( GraphObj )
 	{
-		return GraphObj->GetSearchDepthLimit();
+		return GraphObj->GetSearchDependencyDepthLimit();
 	}
 	else
 	{
@@ -907,11 +993,32 @@ int32 SReferenceViewer::GetSearchDepthCount() const
 	}
 }
 
-void SReferenceViewer::OnSearchDepthCommitted(int32 NewValue)
+int32 SReferenceViewer::GetSearchReferencerDepthCount() const
 {
 	if ( GraphObj )
 	{
-		GraphObj->SetSearchDepthLimit(NewValue);
+		return GraphObj->GetSearchReferencerDepthLimit();
+	}
+	else
+	{
+		return 0;
+	}
+}
+
+void SReferenceViewer::OnSearchDependencyDepthCommitted(int32 NewValue)
+{
+	if ( GraphObj )
+	{
+		GraphObj->SetSearchDependencyDepthLimit(NewValue);
+		RebuildGraph();
+	}
+}
+
+void SReferenceViewer::OnSearchReferencerDepthCommitted(int32 NewValue)
+{
+	if ( GraphObj )
+	{
+		GraphObj->SetSearchReferencerDepthLimit(NewValue);
 		RebuildGraph();
 	}
 }
@@ -1165,6 +1272,29 @@ ECheckBoxState SReferenceViewer::IsCompactModeChecked() const
 	}
 }
 
+void SReferenceViewer::OnShowDuplicatesChanged(ECheckBoxState NewState)
+{
+	if (GraphObj)
+	{
+		GraphObj->SetShowDuplicatesEnabled(NewState == ECheckBoxState::Checked);
+	}
+	
+	RebuildGraph();
+}
+
+
+ECheckBoxState SReferenceViewer::IsShowDuplicatesChecked() const
+{
+	if (GraphObj)
+	{
+		return GraphObj->IsShowDuplicates() ? ECheckBoxState::Checked : ECheckBoxState::Unchecked;
+	}
+	else
+	{
+		return ECheckBoxState::Unchecked;
+	}
+}
+
 void SReferenceViewer::OnShowEditorOnlyReferencesChanged(ECheckBoxState NewState)
 {
 	if (GraphObj)
diff --git a/Engine/Plugins/Editor/AssetManagerEditor/Source/AssetManagerEditor/Private/ReferenceViewer/SReferenceViewer.h b/Engine/Plugins/Editor/AssetManagerEditor/Source/AssetManagerEditor/Private/ReferenceViewer/SReferenceViewer.h
--- a/Engine/Plugins/Editor/AssetManagerEditor/Source/AssetManagerEditor/Private/ReferenceViewer/SReferenceViewer.h
+++ b/Engine/Plugins/Editor/AssetManagerEditor/Source/AssetManagerEditor/Private/ReferenceViewer/SReferenceViewer.h
@@ -103,8 +103,12 @@ private:
 	
 	void OnSearchDepthEnabledChanged( ECheckBoxState NewState );
 	ECheckBoxState IsSearchDepthEnabledChecked() const;
-	int32 GetSearchDepthCount() const;
-	void OnSearchDepthCommitted(int32 NewValue);
+
+	int32 GetSearchReferencerDepthCount() const;
+	int32 GetSearchDependencyDepthCount() const;
+
+	void OnSearchReferencerDepthCommitted(int32 NewValue);
+	void OnSearchDependencyDepthCommitted(int32 NewValue);
 
 	void OnSearchBreadthEnabledChanged( ECheckBoxState NewState );
 	ECheckBoxState IsSearchBreadthEnabledChecked() const;
@@ -130,6 +134,9 @@ private:
 	void OnCompactModeChanged(ECheckBoxState NewState);
 	ECheckBoxState IsCompactModeChecked() const;
 
+	void OnShowDuplicatesChanged(ECheckBoxState NewState);
+	ECheckBoxState IsShowDuplicatesChecked() const;
+
 	EVisibility GetManagementReferencesVisibility() const;
 	void OnShowManagementReferencesChanged(ECheckBoxState NewState);
 	ECheckBoxState IsShowManagementReferencesChecked() const;
diff --git a/Engine/Plugins/Editor/AssetManagerEditor/Source/AssetManagerEditor/Public/ReferenceViewer/EdGraphNode_Reference.h b/Engine/Plugins/Editor/AssetManagerEditor/Source/AssetManagerEditor/Public/ReferenceViewer/EdGraphNode_Reference.h
--- a/Engine/Plugins/Editor/AssetManagerEditor/Source/AssetManagerEditor/Public/ReferenceViewer/EdGraphNode_Reference.h
+++ b/Engine/Plugins/Editor/AssetManagerEditor/Source/AssetManagerEditor/Public/ReferenceViewer/EdGraphNode_Reference.h
@@ -34,6 +34,8 @@ class ASSETMANAGEREDITOR_API UEdGraphNode_Reference : public UEdGraphNode
 	virtual FText GetTooltipText() const override;
 	virtual void AllocateDefaultPins() override;
 	virtual UObject* GetJumpTargetForDoubleClick() const override;
+	virtual FSlateIcon GetIconAndTint(FLinearColor& OutColor) const override;
+	virtual bool ShowPaletteIconOnNode() const override { return true; }
 	// End UEdGraphNode implementation
 
 	void SetAllowThumbnail(bool bInAllow) { bAllowThumbnail = bInAllow; }
@@ -41,6 +43,7 @@ class ASSETMANAGEREDITOR_API UEdGraphNode_Reference : public UEdGraphNode
 	bool UsesThumbnail() const;
 	bool IsPackage() const;
 	bool IsCollapsed() const;
+	bool IsADuplicate() const { return bIsADuplicate; }
 	FAssetData GetAssetData() const;
 
 	UEdGraphPin* GetDependencyPin();
@@ -48,7 +51,7 @@ class ASSETMANAGEREDITOR_API UEdGraphNode_Reference : public UEdGraphNode
 
 private:
 	void CacheAssetData(const FAssetData& AssetData);
-	void SetupReferenceNode(const FIntPoint& NodeLoc, const TArray<FAssetIdentifier>& NewIdentifiers, const FAssetData& InAssetData, bool bInAllowThumbnail);
+	void SetupReferenceNode(const FIntPoint& NodeLoc, const TArray<FAssetIdentifier>& NewIdentifiers, const FAssetData& InAssetData, bool bInAllowThumbnail, bool bInIsADuplicate);
 	void SetReferenceNodeCollapsed(const FIntPoint& NodeLoc, int32 InNumReferencesExceedingMax);
 	void AddReferencer(class UEdGraphNode_Reference* ReferencerNode);
 
@@ -60,7 +63,11 @@ private:
 	bool bIsPackage;
 	bool bIsPrimaryAsset;
 	bool bIsCollapsed;
+	bool bIsADuplicate;
+
 	FAssetData CachedAssetData;
+	FLinearColor AssetTypeColor;
+	FSlateIcon AssetBrush;
 
 	UEdGraphPin* DependencyPin;
 	UEdGraphPin* ReferencerPin;
diff --git a/Engine/Plugins/Editor/AssetManagerEditor/Source/AssetManagerEditor/Public/ReferenceViewer/EdGraph_ReferenceViewer.h b/Engine/Plugins/Editor/AssetManagerEditor/Source/AssetManagerEditor/Public/ReferenceViewer/EdGraph_ReferenceViewer.h
--- a/Engine/Plugins/Editor/AssetManagerEditor/Source/AssetManagerEditor/Public/ReferenceViewer/EdGraph_ReferenceViewer.h
+++ b/Engine/Plugins/Editor/AssetManagerEditor/Source/AssetManagerEditor/Public/ReferenceViewer/EdGraph_ReferenceViewer.h
@@ -15,6 +15,43 @@ class UEdGraphNode_Reference;
 class SReferenceViewer;
 enum class EDependencyPinCategory;
 
+/*
+*  Holds asset information for building reference graph
+*/ 
+struct ReferenceNodeInfo
+{
+	FAssetIdentifier AssetId;
+
+	FAssetData AssetData;
+
+	// immediate children (references or dependencies)
+	TArray<TPair<FAssetIdentifier, EDependencyPinCategory>> Children;
+
+	// this node's parent references (how it got included)
+	TArray<FAssetIdentifier> Parents;
+
+	// Which direction.  Referencers are left (other assets that depend on me), Dependencies are right (other assets I depend on)
+	bool bReferencers;
+
+	int32 OverflowCount;
+
+	ReferenceNodeInfo(const FAssetIdentifier& InAssetId, bool InbReferencers);
+
+	bool IsFirstParent(const FAssetIdentifier& InParentId) const;
+
+	bool IsADuplicate() const;
+
+	// The Provision Size, or vertical spacing required for layout, for a given parent.  
+	// At the time of writing, the intent is only the first node manifestation of 
+	// an asset will have its children shown
+	int32 ProvisionSize(const FAssetIdentifier& InParentId) const;
+
+	// how many nodes worth of children require vertical spacing 
+	int32 ChildProvisionSize;
+
+};
+
+
 UCLASS()
 class ASSETMANAGEREDITOR_API UEdGraph_ReferenceViewer : public UEdGraph
 {
@@ -52,6 +89,7 @@ public:
 	bool IsShowReferencers() const;
 	bool IsShowDependencies() const;
 	bool IsCompactMode() const;
+	bool IsShowDuplicates() const;
 
 	void SetSearchDepthLimitEnabled(bool newEnabled);
 	void SetSearchBreadthLimitEnabled(bool newEnabled);
@@ -65,14 +103,18 @@ public:
 	void SetShowReferencers(const bool bShouldShowReferencers);
 	void SetShowDependencies(const bool bShouldShowDependencies);
 	void SetCompactModeEnabled(bool newEnabled);
+	void SetShowDuplicatesEnabled(bool newEnabled);
 
 	using FIsPackageNamePassingFilterCallback = TFunction<bool(FName)>;
 	void SetIsPackageNamePassingFilterCallback(const TOptional<FIsPackageNamePassingFilterCallback>& InIsPackageNamePassingFilterCallback) { IsPackageNamePassingFilterCallback = InIsPackageNamePassingFilterCallback; }
 
-	int32 GetSearchDepthLimit() const;
-	int32 GetSearchBreadthLimit() const;
+	int32 GetSearchReferencerDepthLimit() const;
+	void SetSearchReferencerDepthLimit(int32 NewDepthLimit);
+
+	int32 GetSearchDependencyDepthLimit() const;
+	void SetSearchDependencyDepthLimit(int32 NewDepthLimit);
 
-	void SetSearchDepthLimit(int32 NewDepthLimit);
+	int32 GetSearchBreadthLimit() const;
 	void SetSearchBreadthLimit(int32 NewBreadthLimit);
 
 	FName GetCurrentCollectionFilter() const;
@@ -81,19 +123,42 @@ public:
 	bool GetEnableCollectionFilter() const;
 	void SetEnableCollectionFilter(bool bEnabled);
 
+	/* Temporary variable that allows reverting to deprecated layout methods */
+	bool GetUseNodeInfos() const { return bUseNodeInfos; }
+	void SetUseNodeInfos(bool InbUseNodeInfos) { bUseNodeInfos= InbUseNodeInfos; }
+
 private:
 	void SetReferenceViewer(TSharedPtr<SReferenceViewer> InViewer);
 	UEdGraphNode_Reference* ConstructNodes(const TArray<FAssetIdentifier>& GraphRootIdentifiers, const FIntPoint& GraphRootOrigin);
-	int32 RecursivelyGatherSizes(bool bReferencers, const TArray<FAssetIdentifier>& Identifiers, const TSet<FName>& AllowedPackageNames, int32 CurrentDepth, TSet<FAssetIdentifier>& VisitedNames, TMap<FAssetIdentifier, int32>& OutNodeSizes) const;
+	// RecursivelyGatherSizes is now deprecated, use RecursivelyPopulateNodeInfos with RecursivelyCreateNodes
+	int32 RecursivelyGatherSizes(bool bReferencers, const TArray<FAssetIdentifier>& Identifiers, const TSet<FName>& AllowedPackageNames, int32 CurrentDepth, int32 MaxDepth, TSet<FAssetIdentifier>& VisitedNames, TMap<FAssetIdentifier, int32>& OutNodeSizes) const;
 	void GatherAssetData(const TSet<FName>& AllPackageNames, TMap<FName, FAssetData>& OutPackageToAssetDataMap) const;
-	class UEdGraphNode_Reference* RecursivelyConstructNodes(bool bReferencers, UEdGraphNode_Reference* RootNode, const TArray<FAssetIdentifier>& Identifiers, const FIntPoint& NodeLoc, const TMap<FAssetIdentifier, int32>& NodeSizes, const TMap<FName, FAssetData>& PackagesToAssetDataMap, const TSet<FName>& AllowedPackageNames, int32 CurrentDepth, TSet<FAssetIdentifier>& VisitedNames);
 
-	bool ExceedsMaxSearchDepth(int32 Depth) const;
+	// RecursivelyConstructNodes is now deprecated, use RecursivelyPopulateNodeInfos with RecursivelyCreateNodes
+	class UEdGraphNode_Reference* RecursivelyConstructNodes(bool bReferencers, UEdGraphNode_Reference* RootNode, const TArray<FAssetIdentifier>& Identifiers, const FIntPoint& NodeLoc, const TMap<FAssetIdentifier, int32>& NodeSizes, const TMap<FName, FAssetData>& PackagesToAssetDataMap, const TSet<FName>& AllowedPackageNames, int32 CurrentDepth, int32 MaxDepth, TSet<FAssetIdentifier>& VisitedNames);
+
+	bool ExceedsMaxSearchDepth(int32 Depth, int32 MaxDepth) const;
 	bool ExceedsMaxSearchBreadth(int32 Breadth) const;
 	FAssetManagerDependencyQuery GetReferenceSearchFlags(bool bHardOnly) const;
 
 	UEdGraphNode_Reference* CreateReferenceNode();
 
+	/* Generates a NodeInfo structure then used to generate and layout the graph nodes */
+	void RecursivelyPopulateNodeInfos(bool bReferencers, const FAssetIdentifier& AssetId, TMap<FAssetIdentifier, ReferenceNodeInfo>& NodeInfos, int32 CurrentDepth, int32 MaxDepth);
+
+	/* Uses the NodeInfos map to generate and layout the graph nodes */
+	UEdGraphNode_Reference* RecursivelyCreateNodes(
+		bool bInReferencers, 
+		const FAssetIdentifier& InAssetId, 
+		const FIntPoint& InNodeLoc, 
+		const FAssetIdentifier& InParentId, 
+		UEdGraphNode_Reference* InParentNode, 
+		TMap<FAssetIdentifier, ReferenceNodeInfo>& InNodeInfos, 
+		int32 InCurrentDepth, 
+		int32 InMaxDepth, 
+		bool bIsRoot = false
+	);
+
 	/** Removes all nodes from the graph */
 	void RemoveAllNodes();
 
@@ -113,7 +178,8 @@ private:
 	TArray<FAssetIdentifier> CurrentGraphRootIdentifiers;
 	FIntPoint CurrentGraphRootOrigin;
 
-	int32 MaxSearchDepth;
+	int32 MaxSearchReferencerDepth; // How deep to search references
+	int32 MaxSearchDependencyDepth; // How deep to search dependanies
 	int32 MaxSearchBreadth;
 
 	/** Current collection filter. NAME_None for no filter */
@@ -132,10 +198,18 @@ private:
 	bool bIsShowReferencers;
 	/* Whether to display the Dependencies */
 	bool bIsShowDependencies;
+	/* Whether to show duplicate asset references */
+	bool bIsShowDuplicates;
+
+	/* This is a convenience toggle to switch between the old & new methods for computing & displaying the graph */
+	bool bUseNodeInfos;
 
 	bool bIsShowFilteredPackagesOnly;
 	TOptional<FIsPackageNamePassingFilterCallback> IsPackageNamePassingFilterCallback;
 
+	/** List of packages the current collection filter allows */
+	TSet<FName> CurrentCollectionPackages;
+
 	bool bIsCompactMode;
 
 	friend SReferenceViewer;
diff --git a/Engine/Source/Editor/GraphEditor/Private/SGraphNode.cpp b/Engine/Source/Editor/GraphEditor/Private/SGraphNode.cpp
--- a/Engine/Source/Editor/GraphEditor/Private/SGraphNode.cpp
+++ b/Engine/Source/Editor/GraphEditor/Private/SGraphNode.cpp
@@ -38,6 +38,8 @@ void SNodeTitle::Construct(const FArguments& InArgs, UEdGraphNode* InNode)
 
 	ExtraLineStyle = InArgs._ExtraLineStyle;
 
+	StyleSet = InArgs._StyleSet;
+
 	CachedSize = FVector2D::ZeroVector;
 
 	// If the user set the text, use it, otherwise use the node title by default
@@ -141,7 +143,7 @@ void SNodeTitle::RebuildWidget()
 		.AutoHeight()
 		[
 			SNew(STextBlock)
-			.TextStyle( FEditorStyle::Get(), ExtraLineStyle )
+			.TextStyle( StyleSet, ExtraLineStyle )
 			.Text(FText::FromString(Lines[Index]))
 		];
 	}
diff --git a/Engine/Source/Editor/GraphEditor/Public/SGraphNode.h b/Engine/Source/Editor/GraphEditor/Public/SGraphNode.h
--- a/Engine/Source/Editor/GraphEditor/Public/SGraphNode.h
+++ b/Engine/Source/Editor/GraphEditor/Public/SGraphNode.h
@@ -19,6 +19,8 @@
 #include "EdGraph/EdGraphNodeUtils.h"
 #include "SNodePanel.h"
 #include "Widgets/Notifications/SErrorText.h"
+#include "Styling/ISlateStyle.h"
+#include "Styling/AppStyle.h"
 
 class FActorDragDropOp;
 class IToolTip;
@@ -34,10 +36,13 @@ class GRAPHEDITOR_API SNodeTitle : public SCompoundWidget
 {
 public:
 	SLATE_BEGIN_ARGS(SNodeTitle)
-		: _Style(TEXT("Graph.Node.NodeTitle"))
+		: _StyleSet(&FAppStyle::Get())
+		, _Style(TEXT("Graph.Node.NodeTitle"))
 		, _ExtraLineStyle(TEXT("Graph.Node.NodeTitleExtraLines"))
 		{}
 
+		SLATE_ARGUMENT(const ISlateStyle*, StyleSet)
+
 		// The style of the text block, which dictates the font, color, and shadow options. Style overrides all other properties!
 		SLATE_ARGUMENT(FName, Style)
 
@@ -67,6 +72,7 @@ protected:
 	TWeakObjectPtr<UEdGraphNode> GraphNode;
 	FNodeTextCache NodeTitleCache;
 	FName ExtraLineStyle;
+	const ISlateStyle* StyleSet;
 
 	/** The cached head title to return */
 	FText CachedHeadTitle;
