Fixed missing include

diff --git a/Engine/Plugins/Editor/AssetManagerEditor/Source/AssetManagerEditor/Private/ReferenceViewer/ReferenceViewerStyle.cpp b/Engine/Plugins/Editor/AssetManagerEditor/Source/AssetManagerEditor/Private/ReferenceViewer/ReferenceViewerStyle.cpp
--- a/Engine/Plugins/Editor/AssetManagerEditor/Source/AssetManagerEditor/Private/ReferenceViewer/ReferenceViewerStyle.cpp
+++ b/Engine/Plugins/Editor/AssetManagerEditor/Source/AssetManagerEditor/Private/ReferenceViewer/ReferenceViewerStyle.cpp
@@ -6,6 +6,8 @@
 #include "Styling/StarshipCoreStyle.h"
 #include "Styling/SlateStyleMacros.h"
 #include "Styling/StyleColors.h"
+#include "Styling/CoreStyle.h"
+
 
 FName FReferenceViewerStyle::StyleName("ReferenceViewerStyle");
 
