ReferenceViewer - Fixes presentation for Management References (PrimaryAssetIds) and Searchable Names (GameplayTags)

diff --git a/Engine/Plugins/Editor/AssetManagerEditor/Source/AssetManagerEditor/Private/ReferenceViewer/EdGraphNode_Reference.cpp b/Engine/Plugins/Editor/AssetManagerEditor/Source/AssetManagerEditor/Private/ReferenceViewer/EdGraphNode_Reference.cpp
--- a/Engine/Plugins/Editor/AssetManagerEditor/Source/AssetManagerEditor/Private/ReferenceViewer/EdGraphNode_Reference.cpp
+++ b/Engine/Plugins/Editor/AssetManagerEditor/Source/AssetManagerEditor/Private/ReferenceViewer/EdGraphNode_Reference.cpp
@@ -38,6 +38,7 @@ void UEdGraphNode_Reference::SetupReferenceNode(const FIntPoint& NodeLoc, const
 	Identifiers = NewIdentifiers;
 	const FAssetIdentifier& First = NewIdentifiers[0];
 	FString MainAssetName = InAssetData.AssetName.ToString();
+	FString AssetTypeName = InAssetData.AssetClass.ToString();
 
 	FAssetToolsModule& AssetToolsModule = FModuleManager::LoadModuleChecked<FAssetToolsModule>(TEXT("AssetTools"));	
 	if (UClass* AssetClass = InAssetData.GetClass())
@@ -58,22 +59,25 @@ void UEdGraphNode_Reference::SetupReferenceNode(const FIntPoint& NodeLoc, const
 	FPrimaryAssetId PrimaryAssetID = NewIdentifiers[0].GetPrimaryAssetId();
 	if (PrimaryAssetID.IsValid())
 	{
-		MainAssetName = PrimaryAssetID.ToString();
+		MainAssetName = PrimaryAssetID.PrimaryAssetName.ToString();
+		AssetTypeName = PrimaryAssetID.PrimaryAssetType.ToString();
 		bIsPackage = false;
 		bIsPrimaryAsset = true;
 	}
 	else if (NewIdentifiers[0].IsValue())
 	{
-		MainAssetName = FString::Printf(TEXT("%s::%s"), *First.ObjectName.ToString(), *First.ValueName.ToString());
+		MainAssetName = First.ValueName.ToString();
+		AssetTypeName = First.ObjectName.ToString();
 		bIsPackage = false;
 	}
 
 	if (NewIdentifiers.Num() == 1 )
 	{
 		static const FName NAME_ActorLabel(TEXT("ActorLabel"));
-		
-		NodeTitle = FText::FromString(FString::Printf(TEXT("%s\n%s"), *InAssetData.AssetName.ToString(), *InAssetData.AssetClass.ToString()));
-		InAssetData.GetTagValue(NAME_ActorLabel, NodeTitle);
+		InAssetData.GetTagValue(NAME_ActorLabel, MainAssetName); 
+
+		// append the type so it shows up on the extra line
+		NodeTitle = FText::FromString(FString::Printf(TEXT("%s\n%s"), *MainAssetName, *AssetTypeName));
 	}
 	else
 	{
