Reference Viewer UI

- Moved view options into toolbar
- Migrated Functions to Commands
- Added Keyboard Shortcuts for View Options
- Added Keyboard Shortcuts to adjust Depth & Breadth
- Added Keyboard Shortcut to edit the selected node

diff --git a/Engine/Plugins/Editor/AssetManagerEditor/Source/AssetManagerEditor/AssetManagerEditor.Build.cs b/Engine/Plugins/Editor/AssetManagerEditor/Source/AssetManagerEditor/AssetManagerEditor.Build.cs
--- a/Engine/Plugins/Editor/AssetManagerEditor/Source/AssetManagerEditor/AssetManagerEditor.Build.cs
+++ b/Engine/Plugins/Editor/AssetManagerEditor/Source/AssetManagerEditor/AssetManagerEditor.Build.cs
@@ -45,6 +45,7 @@ public class AssetManagerEditor : ModuleRules
 				"EditorWidgets",
 				"TreeMap",
 				"ToolMenus",
+				"ToolWidgets",
 			}
 		);
 	}
diff --git a/Engine/Plugins/Editor/AssetManagerEditor/Source/AssetManagerEditor/Private/AssetManagerEditorCommands.cpp b/Engine/Plugins/Editor/AssetManagerEditor/Source/AssetManagerEditor/Private/AssetManagerEditorCommands.cpp
--- a/Engine/Plugins/Editor/AssetManagerEditor/Source/AssetManagerEditor/Private/AssetManagerEditorCommands.cpp
+++ b/Engine/Plugins/Editor/AssetManagerEditor/Source/AssetManagerEditor/Private/AssetManagerEditorCommands.cpp
@@ -23,8 +23,33 @@ void FAssetManagerEditorCommands::RegisterCommands()
 	
 	UI_COMMAND(ViewAssetAudit, "Audit Assets...", "Opens the Asset Audit UI and displays information about the selected assets", EUserInterfaceActionType::Button, FInputChord(EModifierKey::Shift | EModifierKey::Alt, EKeys::A));
 
-	UI_COMMAND(OpenSelectedInAssetEditor, "Edit...", "Opens the selected asset in the relevant editor.", EUserInterfaceActionType::Button, FInputChord());
+	UI_COMMAND(OpenSelectedInAssetEditor, "Edit...", "Opens the selected asset in the relevant editor.", EUserInterfaceActionType::Button, FInputChord(EModifierKey::Control, EKeys::E));
 	UI_COMMAND(ZoomToFit, "Zoom to Fit", "Zoom in and center the view on the selected item", EUserInterfaceActionType::Button, FInputChord(EKeys::Home));
+
+	UI_COMMAND(IncreaseReferencerSearchDepth, "Increase Referencer Search Depth", "Increase the Referencer Search Depth", EUserInterfaceActionType::Button, FInputChord(EKeys::R));
+	UI_COMMAND(DecreaseReferencerSearchDepth, "Decrease Referencer Search Depth", "Decrease the Referencer Search Depth", EUserInterfaceActionType::Button, FInputChord(EModifierKey::Shift, EKeys::R));
+	UI_COMMAND(SetReferencerSearchDepth, "Set Referencer Search Depth", "Set the Referencer Search Depth", EUserInterfaceActionType::Button, FInputChord(EModifierKey::Control, EKeys::R));
+
+	UI_COMMAND(IncreaseDependencySearchDepth, "Increase Dependency Search Depth", "Increase the Dependency Search Depth", EUserInterfaceActionType::Button, FInputChord(EKeys::D));
+	UI_COMMAND(DecreaseDependencySearchDepth, "Decrease Dependency Search Depth", "Decrease the Dependency Search Depth", EUserInterfaceActionType::Button, FInputChord(EModifierKey::Shift, EKeys::D));
+	UI_COMMAND(SetDependencySearchDepth, "Set Dependency Search Depth", "Set the Dependency Search Depth", EUserInterfaceActionType::Button, FInputChord(EModifierKey::Control, EKeys::D));
+
+	UI_COMMAND(IncreaseBreadth, "Increase Search Breadth", "Increase the Breadth Limit", EUserInterfaceActionType::Button, FInputChord( EKeys::L));
+	UI_COMMAND(DecreaseBreadth, "Decrease Search Breadth", "Decrease the Breadth Limit", EUserInterfaceActionType::Button, FInputChord(EModifierKey::Shift, EKeys::L));
+	UI_COMMAND(SetBreadth, "Set the Breadth", "Set the Breadth Limit", EUserInterfaceActionType::Button, FInputChord(EModifierKey::Control, EKeys::L));
+
+	UI_COMMAND(ShowSoftReferences, "Show Soft References", "Toggles visibility of Soft References", EUserInterfaceActionType::ToggleButton, FInputChord(EKeys::S));
+	UI_COMMAND(ShowHardReferences, "Show Hard References", "Toggles visibility of Hard References", EUserInterfaceActionType::ToggleButton, FInputChord(EKeys::H));
+	UI_COMMAND(ShowEditorOnlyReferences, "Show EditorOnly References","Toggles visibility of EditorOnly References", EUserInterfaceActionType::ToggleButton, FInputChord(EKeys::E));
+
+	UI_COMMAND(ShowManagementReferences, "Show Management References","Toggles visibility of Management References (i.e. PrimaryAssetIDs)", EUserInterfaceActionType::ToggleButton, FInputChord(EKeys::M));
+	UI_COMMAND(ShowGamePlayTags, "Show GamePlayTags","Toggles visibility of GameplayTags", EUserInterfaceActionType::ToggleButton, FInputChord(EKeys::G));
+	UI_COMMAND(ShowNativePackages, "Show Native Packages","Toggles visibility of Native Packages", EUserInterfaceActionType::ToggleButton, FInputChord(EKeys::N));
+
+	UI_COMMAND(ShowDuplicates, "Show Duplicate References", "Toggles visibility of Duplicate References", EUserInterfaceActionType::ToggleButton, FInputChord(EKeys::U));
+	UI_COMMAND(CompactMode, "Compact Mode", "Toggles Compact Mode", EUserInterfaceActionType::ToggleButton, FInputChord(EKeys::C));
+	UI_COMMAND(FilterSearch, "Filter Search Results", "Toggles filtering of search results", EUserInterfaceActionType::ToggleButton, FInputChord(EKeys::F));
+
 	UI_COMMAND(ReCenterGraph, "Re-Center Graph", "Re-centers the graph on this node, showing all referencers and references for this asset instead", EUserInterfaceActionType::Button, FInputChord());
 	UI_COMMAND(CopyReferencedObjects, "Copy Referenced Objects List", "Copies the list of objects that the selected asset references to the clipboard.", EUserInterfaceActionType::Button, FInputChord());
 	UI_COMMAND(CopyReferencingObjects, "Copy Referencing Objects List", "Copies the list of objects that reference the selected asset to the clipboard.", EUserInterfaceActionType::Button, FInputChord());
diff --git a/Engine/Plugins/Editor/AssetManagerEditor/Source/AssetManagerEditor/Private/AssetManagerEditorCommands.h b/Engine/Plugins/Editor/AssetManagerEditor/Source/AssetManagerEditor/Private/AssetManagerEditorCommands.h
--- a/Engine/Plugins/Editor/AssetManagerEditor/Source/AssetManagerEditor/Private/AssetManagerEditorCommands.h
+++ b/Engine/Plugins/Editor/AssetManagerEditor/Source/AssetManagerEditor/Private/AssetManagerEditorCommands.h
@@ -34,6 +34,60 @@ public:
 	// Re-constructs the graph with the selected asset as the center
 	TSharedPtr<FUICommandInfo> ReCenterGraph;
 
+	// Increase Referencer Search Depth
+	TSharedPtr<FUICommandInfo> IncreaseReferencerSearchDepth;
+
+	// Decrease Referencer Search Depth
+	TSharedPtr<FUICommandInfo> DecreaseReferencerSearchDepth;
+
+	// Set the Referencer Serach Depth 
+	TSharedPtr<FUICommandInfo> SetReferencerSearchDepth;
+
+	// Increase Dependency Search Depth
+	TSharedPtr<FUICommandInfo> IncreaseDependencySearchDepth;
+
+	// Decrease Dependency Search Depth
+	TSharedPtr<FUICommandInfo> DecreaseDependencySearchDepth;
+
+	// Set the Dependency Serach Depth 
+	TSharedPtr<FUICommandInfo> SetDependencySearchDepth;
+
+	// Increase Search Breadth
+	TSharedPtr<FUICommandInfo> IncreaseBreadth;
+
+	// Decrease Search Breadth
+	TSharedPtr<FUICommandInfo> DecreaseBreadth;
+
+	// Set the Breadth LImit
+	TSharedPtr<FUICommandInfo> SetBreadth;;
+
+	// Toggles visiblity of Soft References
+	TSharedPtr<FUICommandInfo> ShowSoftReferences;
+
+	// Toggles visiblity of Hard References
+	TSharedPtr<FUICommandInfo> ShowHardReferences;
+
+	// Toggles visiblity of EditorOnly References
+	TSharedPtr<FUICommandInfo> ShowEditorOnlyReferences;
+
+	// Toggles visiblity of Management References (i.e. PrimaryAssetIDs)
+	TSharedPtr<FUICommandInfo> ShowManagementReferences;
+
+	// Toggles visiblity of GameplayTags
+	TSharedPtr<FUICommandInfo> ShowGamePlayTags;
+
+	// Toggles visiblity of Native Packages
+	TSharedPtr<FUICommandInfo> ShowNativePackages;
+
+	// Toggles visiblity of Duplicate References
+	TSharedPtr<FUICommandInfo> ShowDuplicates;
+
+	// Toggles Compact Mode
+	TSharedPtr<FUICommandInfo> CompactMode;
+
+	// Toggles whether search results are filtered or just selected 
+	TSharedPtr<FUICommandInfo> FilterSearch;
+
 	// Copies the list of objects that the selected asset references
 	TSharedPtr<FUICommandInfo> CopyReferencedObjects;
 
diff --git a/Engine/Plugins/Editor/AssetManagerEditor/Source/AssetManagerEditor/Private/ReferenceViewer/EdGraph_ReferenceViewer.cpp b/Engine/Plugins/Editor/AssetManagerEditor/Source/AssetManagerEditor/Private/ReferenceViewer/EdGraph_ReferenceViewer.cpp
--- a/Engine/Plugins/Editor/AssetManagerEditor/Source/AssetManagerEditor/Private/ReferenceViewer/EdGraph_ReferenceViewer.cpp
+++ b/Engine/Plugins/Editor/AssetManagerEditor/Source/AssetManagerEditor/Private/ReferenceViewer/EdGraph_ReferenceViewer.cpp
@@ -276,7 +276,7 @@ int32 UEdGraph_ReferenceViewer::GetSearchDependencyDepthLimit() const
 
 void UEdGraph_ReferenceViewer::SetSearchDependencyDepthLimit(int32 NewDepthLimit)
 {
-	MaxSearchDependencyDepth = NewDepthLimit;
+	MaxSearchDependencyDepth = FMath::Max(NewDepthLimit, 0);
 }
 
 int32 UEdGraph_ReferenceViewer::GetSearchReferencerDepthLimit() const
@@ -286,7 +286,7 @@ int32 UEdGraph_ReferenceViewer::GetSearchReferencerDepthLimit() const
 
 void UEdGraph_ReferenceViewer::SetSearchReferencerDepthLimit(int32 NewDepthLimit)
 {
-	MaxSearchReferencerDepth = NewDepthLimit;
+	MaxSearchReferencerDepth = FMath::Max(NewDepthLimit, 0);
 }
 
 int32 UEdGraph_ReferenceViewer::GetSearchBreadthLimit() const
@@ -296,7 +296,7 @@ int32 UEdGraph_ReferenceViewer::GetSearchBreadthLimit() const
 
 void UEdGraph_ReferenceViewer::SetSearchBreadthLimit(int32 NewBreadthLimit)
 {
-	MaxSearchBreadth = NewBreadthLimit;
+	MaxSearchBreadth = FMath::Max(NewBreadthLimit, 0);
 }
 
 FName UEdGraph_ReferenceViewer::GetCurrentCollectionFilter() const
diff --git a/Engine/Plugins/Editor/AssetManagerEditor/Source/AssetManagerEditor/Private/ReferenceViewer/SReferenceViewer.cpp b/Engine/Plugins/Editor/AssetManagerEditor/Source/AssetManagerEditor/Private/ReferenceViewer/SReferenceViewer.cpp
--- a/Engine/Plugins/Editor/AssetManagerEditor/Source/AssetManagerEditor/Private/ReferenceViewer/SReferenceViewer.cpp
+++ b/Engine/Plugins/Editor/AssetManagerEditor/Source/AssetManagerEditor/Private/ReferenceViewer/SReferenceViewer.cpp
@@ -17,6 +17,9 @@
 #include "Widgets/Input/SCheckBox.h"
 #include "Widgets/Input/SSearchBox.h"
 #include "Widgets/Input/SSpinBox.h"
+#include "SSimpleButton.h"
+#include "SSimpleComboButton.h"
+
 #include "Styling/AppStyle.h"
 #include "Engine/Selection.h"
 #include "ReferenceViewer/EdGraph_ReferenceViewer.h"
@@ -34,6 +37,7 @@
 #include "HAL/PlatformApplicationMisc.h"
 #include "AssetManagerEditorModule.h"
 #include "Framework/Application/SlateApplication.h"
+#include "Framework/MultiBox/MultiBoxBuilder.h"
 
 #include "ObjectTools.h"
 
@@ -122,6 +126,8 @@ void SReferenceViewer::Construct(const FArguments& InArgs)
 	FEditorWidgetsModule& EditorWidgetsModule = FModuleManager::LoadModuleChecked<FEditorWidgetsModule>("EditorWidgets");
 	TSharedRef<SWidget> AssetDiscoveryIndicator = EditorWidgetsModule.CreateAssetDiscoveryIndicator(EAssetDiscoveryIndicatorScaleMode::Scale_None, FMargin(16, 8), false);
 
+	const FAssetManagerEditorCommands& UICommands	= FAssetManagerEditorCommands::Get();
+
 	static const FName DefaultForegroundName("DefaultForeground");
 
 	// Visual options visibility
@@ -137,84 +143,99 @@ void SReferenceViewer::Construct(const FArguments& InArgs)
 
 	ChildSlot
 	[
+
 		SNew(SVerticalBox)
 
 		// Path and history
 		+SVerticalBox::Slot()
 		.AutoHeight()
-		.Padding( 0, 0, 0, 4 )
 		[
-			SNew(SHorizontalBox)
-
-			// History Back Button
-			+SHorizontalBox::Slot()
-			.AutoWidth()
-			.VAlign(VAlign_Center)
-			.Padding(1,0)
+			SNew(SBorder)
+			.BorderImage(FAppStyle::GetBrush("Brushes.Panel"))
+			.Padding(FMargin(12, 6, 12, 6))
 			[
-				SNew(SButton)
-				.ButtonStyle( FAppStyle::Get(), "HoverHintOnly" )
-				.ForegroundColor( FAppStyle::GetSlateColor(DefaultForegroundName) )
-				.ToolTipText( this, &SReferenceViewer::GetHistoryBackTooltip )
-				.ContentPadding( 0 )
-				.OnClicked(this, &SReferenceViewer::BackClicked)
-				.IsEnabled(this, &SReferenceViewer::IsBackEnabled)
+				SNew(SHorizontalBox)
+
+
+				// Refresh Button
+				+SHorizontalBox::Slot()
+				.AutoWidth()
+				.Padding(4, 0)
 				[
-					SNew(SImage) .Image(FAppStyle::GetBrush("Icons.CircleArrowLeft"))
+					SNew(SSimpleButton)
+					.ToolTipText(LOCTEXT("RefreshTooltip", "Refresh current view"))
+					.OnClicked(this, &SReferenceViewer::RefreshClicked)
+					.Icon(FAppStyle::GetBrush("Icons.Refresh"))
 				]
-			]
 
-			// History Forward Button
-			+SHorizontalBox::Slot()
-			.AutoWidth()
-			.Padding(1,0,3,0)
-			.VAlign(VAlign_Center)
-			[
-				SNew(SButton)
-				.ButtonStyle( FAppStyle::Get(), "HoverHintOnly" )
-				.ForegroundColor( FAppStyle::GetSlateColor(DefaultForegroundName) )
-				.ToolTipText( this, &SReferenceViewer::GetHistoryForwardTooltip )
-				.ContentPadding( 0 )
-				.OnClicked(this, &SReferenceViewer::ForwardClicked)
-				.IsEnabled(this, &SReferenceViewer::IsForwardEnabled)
+
+				// History Back Button
+				+SHorizontalBox::Slot()
+				.AutoWidth()
+				.Padding(4, 0)
 				[
-					SNew(SImage) .Image(FAppStyle::GetBrush("Icons.CircleArrowRight"))
+					SNew(SButton)
+					.ToolTipText( this, &SReferenceViewer::GetHistoryBackTooltip )
+					.ButtonStyle(&FAppStyle::Get().GetWidgetStyle<FButtonStyle>("SimpleButton"))
+					.ForegroundColor(FSlateColor::UseStyle())
+					.OnClicked(this, &SReferenceViewer::BackClicked)
+					.IsEnabled(this, &SReferenceViewer::IsBackEnabled)
+					[
+						SNew(SImage)
+						.ColorAndOpacity(FSlateColor::UseForeground())
+						.Image(FAppStyle::GetBrush("Icons.ArrowLeft"))
+						.DesiredSizeOverride(FVector2D(20.0f, 20.0f))
+					]
 				]
-			]
 
-			// Refresh Button
-			+SHorizontalBox::Slot()
-			.AutoWidth()
-			.Padding(1,0,3,0)
-			.VAlign(VAlign_Center)
-			[
-				SNew(SButton)
-				.ButtonStyle( FAppStyle::Get(), "HoverHintOnly" )
-				.ForegroundColor( FAppStyle::GetSlateColor(DefaultForegroundName) )
-				.ToolTipText(LOCTEXT("RefreshTooltip", "Refresh current view"))
-				.ContentPadding( 0 )
-				.OnClicked(this, &SReferenceViewer::RefreshClicked)
+				// History Forward Button
+				+SHorizontalBox::Slot()
+				.AutoWidth()
+				.Padding(4, 0)
 				[
-					SNew(SImage) .Image(FAppStyle::GetBrush("Icons.Refresh"))
+					SNew(SButton)
+					.ToolTipText( this, &SReferenceViewer::GetHistoryForwardTooltip )
+					.ButtonStyle(&FAppStyle::Get().GetWidgetStyle<FButtonStyle>("SimpleButton"))
+					.ForegroundColor(FSlateColor::UseStyle())
+					.OnClicked(this, &SReferenceViewer::ForwardClicked)
+					.IsEnabled(this, &SReferenceViewer::IsForwardEnabled)
+					[
+						SNew(SImage)
+						.ColorAndOpacity(FSlateColor::UseForeground())
+						.Image(FAppStyle::GetBrush("Icons.ArrowRight"))
+						.DesiredSizeOverride(FVector2D(20.0f, 20.0f))
+					]
 				]
-			]
 
-			// Path
-			+SHorizontalBox::Slot()
-			.VAlign(VAlign_Fill)
-			.FillWidth(1.f)
-			[
-				SNew(SBorder)
-				.BorderImage( FAppStyle::GetBrush("ToolPanel.GroupBorder") )
+				+SHorizontalBox::Slot()
+				.AutoWidth()
+				.Padding(4, 0)
 				[
-					SNew(SEditableTextBox)
-					.Text(this, &SReferenceViewer::GetAddressBarText)
-					.OnTextCommitted(this, &SReferenceViewer::OnAddressBarTextCommitted)
-					.OnTextChanged(this, &SReferenceViewer::OnAddressBarTextChanged)
-					.SelectAllTextWhenFocused(true)
-					.SelectAllTextOnCommit(true)
-					.Style(FAppStyle::Get(), "ReferenceViewer.PathText")
+					SNew(SSimpleComboButton)
+					.OnGetMenuContent(this, &SReferenceViewer::GetShowMenuContent)
+					.Icon(FAppStyle::GetBrush("Icons.Visibility"))
+					.HasDownArrow(true)
 				]
+
+				// Path
+				+SHorizontalBox::Slot()
+				.Padding(4, 0)
+				.FillWidth(1.f)
+				[
+					SNew(SBorder)
+					.BorderImage( FAppStyle::GetBrush("ToolPanel.GroupBorder") )
+					[
+						SNew(SEditableTextBox)
+						.Text(this, &SReferenceViewer::GetAddressBarText)
+						.OnTextCommitted(this, &SReferenceViewer::OnAddressBarTextCommitted)
+						.OnTextChanged(this, &SReferenceViewer::OnAddressBarTextChanged)
+						.SelectAllTextWhenFocused(true)
+						.SelectAllTextOnCommit(true)
+						.Style(FAppStyle::Get(), "ReferenceViewer.PathText")
+					]
+				]
+
+			
 			]
 		]
 
@@ -251,31 +272,6 @@ void SReferenceViewer::Construct(const FArguments& InArgs)
 						.OnTextCommitted(this, &SReferenceViewer::HandleOnSearchTextCommitted)
 					]
 
-					+ SVerticalBox::Slot()
-					.AutoHeight()
-					[
-						SNew(SHorizontalBox)
-						.Visibility_Lambda([this]() { return (bShowShowFilteredPackagesOnly ? EVisibility::Visible : EVisibility::Collapsed); })
-
-						+ SHorizontalBox::Slot()
-						.VAlign(VAlign_Center)
-						.Padding(2.f)
-						[
-							SNew(STextBlock)
-							.Text(LOCTEXT("ShowHideFilteredPackagesOnly", "Show Filtered Packages Only"))
-						]
-
-						+SHorizontalBox::Slot()
-						.AutoWidth()
-						.VAlign(VAlign_Center)
-						.Padding(2.f)
-						[
-							SNew(SCheckBox)
-							.OnCheckStateChanged(this, &SReferenceViewer::OnShowFilteredPackagesOnlyChanged)
-							.IsChecked(this, &SReferenceViewer::IsShowFilteredPackagesOnlyChecked)
-						]
-					]
-
 					+SVerticalBox::Slot()
 					.AutoHeight()
 					[
@@ -288,6 +284,11 @@ void SReferenceViewer::Construct(const FArguments& InArgs)
 						[
 							SNew(STextBlock)
 							.Text(LOCTEXT("SearchDepthReferencersLabelText", "Search Referencers Depth"))
+							.ToolTipText(FText::Format( LOCTEXT("ReferenceDepthToolTip", "Adjust Referencer Search Depth (+/-):  {0} / {1}\nSet Referencer Search Depth:                        {2}"),
+															UICommands.IncreaseReferencerSearchDepth->GetInputText().ToUpper(),
+															UICommands.DecreaseReferencerSearchDepth->GetInputText().ToUpper(),
+															UICommands.SetReferencerSearchDepth->GetInputText().ToUpper()))
+
 						]
 
 						+SHorizontalBox::Slot()
@@ -298,9 +299,10 @@ void SReferenceViewer::Construct(const FArguments& InArgs)
 							SNew(SBox)
 							.WidthOverride(100)
 							[
-								SNew(SSpinBox<int32>)
+								SAssignNew(ReferencerCountBox, SSpinBox<int32>)
 								.Value(this, &SReferenceViewer::GetSearchReferencerDepthCount)
 								.OnValueChanged(this, &SReferenceViewer::OnSearchReferencerDepthCommitted)
+								.OnValueCommitted_Lambda([this] (int32 NewValue, ETextCommit::Type CommitType) { FSlateApplication::Get().SetKeyboardFocus(GraphEditorPtr, EFocusCause::SetDirectly); } )
 								.MinValue(0)
 								.MaxValue(50)
 								.MaxSliderValue(12)
@@ -320,6 +322,10 @@ void SReferenceViewer::Construct(const FArguments& InArgs)
 						[
 							SNew(STextBlock)
 							.Text(LOCTEXT("SearchDepthDependenciesLabelText", "Search Dependencies Depth"))
+							.ToolTipText(FText::Format( LOCTEXT("DependencyDepthToolTip", "Adjust Dependency Search Depth (+/-):  {0} / {1}\nSet Dependency Search Depth:                        {2}"),
+															UICommands.IncreaseDependencySearchDepth->GetInputText().ToUpper(),
+															UICommands.DecreaseDependencySearchDepth->GetInputText().ToUpper(),
+															UICommands.SetDependencySearchDepth->GetInputText().ToUpper()))
 						]
 
 						+SHorizontalBox::Slot()
@@ -330,9 +336,10 @@ void SReferenceViewer::Construct(const FArguments& InArgs)
 							SNew(SBox)
 							.WidthOverride(100)
 							[
-								SNew(SSpinBox<int32>)
+								SAssignNew(DependencyCountBox, SSpinBox<int32>)
 								.Value(this, &SReferenceViewer::GetSearchDependencyDepthCount)
 								.OnValueChanged(this, &SReferenceViewer::OnSearchDependencyDepthCommitted)
+								.OnValueCommitted_Lambda([this] (int32 NewValue, ETextCommit::Type CommitType) { FSlateApplication::Get().SetKeyboardFocus(GraphEditorPtr, EFocusCause::SetDirectly); } )
 								.MinValue(0)
 								.MaxValue(50)
 								.MaxSliderValue(12)
@@ -352,6 +359,10 @@ void SReferenceViewer::Construct(const FArguments& InArgs)
 						[
 							SNew(STextBlock)
 							.Text(LOCTEXT("SearchBreadthLabelText", "Search Breadth Limit"))
+							.ToolTipText(FText::Format( LOCTEXT("BreadthLimitToolTip", "Adjust Breadth Limit (+/-):  {0} / {1}\nSet Breadth Limit:                        {2}"),
+															UICommands.IncreaseBreadth->GetInputText().ToUpper(),
+															UICommands.DecreaseBreadth->GetInputText().ToUpper(),
+															UICommands.SetBreadth->GetInputText().ToUpper()))
 						]
 
 						+SHorizontalBox::Slot()
@@ -372,9 +383,10 @@ void SReferenceViewer::Construct(const FArguments& InArgs)
 							SNew(SBox)
 							.WidthOverride(100)
 							[
-								SNew(SSpinBox<int32>)
+								SAssignNew(BreadthLimitBox, SSpinBox<int32>)
 								.Value(this, &SReferenceViewer::GetSearchBreadthCount)
 								.OnValueChanged(this, &SReferenceViewer::OnSearchBreadthCommitted)
+								.OnValueCommitted_Lambda([this] (int32 NewValue, ETextCommit::Type CommitType) { FSlateApplication::Get().SetKeyboardFocus(GraphEditorPtr, EFocusCause::SetDirectly); } )
 								.MinValue(1)
 								.MaxValue(1000)
 								.MaxSliderValue(50)
@@ -428,205 +440,6 @@ void SReferenceViewer::Construct(const FArguments& InArgs)
 						]
 					]
 
-					+SVerticalBox::Slot()
-					.AutoHeight()
-					[
-						SNew(SHorizontalBox)
-						.Visibility_Lambda([this]() { return (bShowShowReferencesOptions ? EVisibility::Visible : EVisibility::Collapsed); })
-
-						+SHorizontalBox::Slot()
-						.VAlign(VAlign_Center)
-						.Padding(2.f)
-						[
-							SNew(STextBlock)
-							.Text(LOCTEXT("ShowHideSoftReferences", "Show Soft References"))
-						]
-
-						+SHorizontalBox::Slot()
-						.AutoWidth()
-						.VAlign(VAlign_Center)
-						.Padding(2.f)
-						[
-							SNew(SCheckBox)
-							.OnCheckStateChanged( this, &SReferenceViewer::OnShowSoftReferencesChanged )
-							.IsChecked( this, &SReferenceViewer::IsShowSoftReferencesChecked )
-						]
-					]
-
-				
-					+ SVerticalBox::Slot()
-					.AutoHeight()
-					[
-						SNew(SHorizontalBox)
-						.Visibility_Lambda([this]() { return (bShowShowReferencesOptions ? EVisibility::Visible : EVisibility::Collapsed); })
-
-						+ SHorizontalBox::Slot()
-						.VAlign(VAlign_Center)
-						.Padding(2.f)
-						[
-							SNew(STextBlock)
-							.Text(LOCTEXT("ShowHideHardReferences", "Show Hard References"))
-						]
-
-						+SHorizontalBox::Slot()
-						.AutoWidth()
-						.VAlign(VAlign_Center)
-						.Padding(2.f)
-						[
-							SNew(SCheckBox)
-							.OnCheckStateChanged(this, &SReferenceViewer::OnShowHardReferencesChanged)
-							.IsChecked(this, &SReferenceViewer::IsShowHardReferencesChecked)
-						]
-					]
-
-					+ SVerticalBox::Slot()
-					.AutoHeight()
-					[
-						SNew(SHorizontalBox)
-						.Visibility_Lambda([this]() { return (bShowShowReferencesOptions ? EVisibility::Visible : EVisibility::Collapsed); })
-
-						+ SHorizontalBox::Slot()
-						.VAlign(VAlign_Center)
-						.Padding(2.f)
-						[
-							SNew(STextBlock)
-							.Text(LOCTEXT("ShowHideEditorOnlyReferences", "Show EditorOnly References"))
-						]
-
-						+SHorizontalBox::Slot()
-						.AutoWidth()
-						.VAlign(VAlign_Center)
-						.Padding(2.f)
-						[
-							SNew(SCheckBox)
-							.OnCheckStateChanged(this, &SReferenceViewer::OnShowEditorOnlyReferencesChanged)
-							.IsChecked(this, &SReferenceViewer::IsShowEditorOnlyReferencesChecked)
-						]
-					]
-
-					+ SVerticalBox::Slot()
-					.AutoHeight()
-					[
-						SNew(SHorizontalBox)
-						.Visibility(this, &SReferenceViewer::GetManagementReferencesVisibility)
-
-						+ SHorizontalBox::Slot()
-						.VAlign(VAlign_Center)
-						.Padding(2.f)
-						[
-							SNew(STextBlock)
-							.Text(LOCTEXT("ShowHideManagementReferences", "Show Management References"))
-						]
-
-						+SHorizontalBox::Slot()
-						.AutoWidth()
-						.VAlign(VAlign_Center)
-						.Padding(2.f)
-						[
-							SNew(SCheckBox)
-							.OnCheckStateChanged(this, &SReferenceViewer::OnShowManagementReferencesChanged)
-							.IsChecked(this, &SReferenceViewer::IsShowManagementReferencesChecked)
-						]
-					]
-
-					+ SVerticalBox::Slot()
-					.AutoHeight()
-					[
-						SNew(SHorizontalBox)
-						.Visibility_Lambda([this]() { return (bShowShowSearchableNames ? EVisibility::Visible : EVisibility::Collapsed); })
-
-						+ SHorizontalBox::Slot()
-						.VAlign(VAlign_Center)
-						.Padding(2.f)
-						[
-							SNew(STextBlock)
-							.Text(LOCTEXT("ShowHideSearchableNames", "Show Searchable Names"))
-						]
-
-						+SHorizontalBox::Slot()
-						.AutoWidth()
-						.VAlign(VAlign_Center)
-						.Padding(2.f)
-						[
-							SNew(SCheckBox)
-							.OnCheckStateChanged(this, &SReferenceViewer::OnShowSearchableNamesChanged)
-							.IsChecked(this, &SReferenceViewer::IsShowSearchableNamesChecked)
-						]
-					]
-
-					+ SVerticalBox::Slot()
-					.AutoHeight()
-					[
-						SNew(SHorizontalBox)
-						.Visibility_Lambda([this]() { return (bShowShowNativePackages ? EVisibility::Visible : EVisibility::Collapsed); })
-
-						+ SHorizontalBox::Slot()
-						.VAlign(VAlign_Center)
-						.Padding(2.f)
-						[
-							SNew(STextBlock)
-							.Text(LOCTEXT("ShowHideNativePackages", "Show Native Packages"))
-						]
-
-						+SHorizontalBox::Slot()
-						.AutoWidth()
-						.VAlign(VAlign_Center)
-						.Padding(2.f)
-						[
-							SNew(SCheckBox)
-							.OnCheckStateChanged(this, &SReferenceViewer::OnShowNativePackagesChanged)
-							.IsChecked(this, &SReferenceViewer::IsShowNativePackagesChecked)
-						]
-					]
-
-					+ SVerticalBox::Slot()
-					.AutoHeight()
-					[
-						SNew(SHorizontalBox)
-						.Visibility_Lambda([this]() { return (bShowCompactMode ? EVisibility::Visible : EVisibility::Collapsed); })
-
-						+ SHorizontalBox::Slot()
-						.VAlign(VAlign_Center)
-						.Padding(2.f)
-						[
-							SNew(STextBlock)
-							.Text(LOCTEXT("CompactMode", "Compact Mode"))
-						]
-
-						+SHorizontalBox::Slot()
-						.AutoWidth()
-						.VAlign(VAlign_Center)
-						.Padding(2.f)
-						[
-							SNew(SCheckBox)
-							.OnCheckStateChanged(this, &SReferenceViewer::OnCompactModeChanged)
-							.IsChecked(this, &SReferenceViewer::IsCompactModeChecked)
-						]
-					]
-
-					+ SVerticalBox::Slot()
-					.AutoHeight()
-					[
-						SNew(SHorizontalBox)
-						+ SHorizontalBox::Slot()
-						.VAlign(VAlign_Center)
-						.Padding(2.f)
-						[
-							SNew(STextBlock)
-							.Text(LOCTEXT("DuplicateAssets", "Show Duplicate Assets"))
-						]
-
-						+SHorizontalBox::Slot()
-						.AutoWidth()
-						.VAlign(VAlign_Center)
-						.Padding(2.f)
-						[
-							SNew(SCheckBox)
-							.OnCheckStateChanged(this, &SReferenceViewer::OnShowDuplicatesChanged)
-							.IsChecked(this, &SReferenceViewer::IsShowDuplicatesChecked)
-						]
-					]
-
 					+ SVerticalBox::Slot()
 					.AutoHeight()
 					[
@@ -685,6 +498,11 @@ void SReferenceViewer::Construct(const FArguments& InArgs)
 	UpdateCollectionsComboList();
 }
 
+FReply SReferenceViewer::OnKeyDown(const FGeometry& MyGeometry, const FKeyEvent& InKeyEvent)
+{
+	return ReferenceViewerActions->ProcessCommandBindings(InKeyEvent) ? FReply::Handled() : FReply::Unhandled();
+}
+
 void SReferenceViewer::SetGraphRootIdentifiers(const TArray<FAssetIdentifier>& NewGraphRootIdentifiers, const FReferenceViewerParams& ReferenceViewerParams)
 {
 	GraphObj->SetGraphRoot(NewGraphRootIdentifiers);
@@ -1159,72 +977,50 @@ FText SReferenceViewer::GetCollectionFilterText() const
 	return FText::FromName(GraphObj->GetCurrentCollectionFilter());
 }
 
-void SReferenceViewer::OnShowSoftReferencesChanged(ECheckBoxState NewState)
+void SReferenceViewer::OnShowSoftReferencesChanged()
 {
 	if (GraphObj)
 	{
-		GraphObj->SetShowSoftReferencesEnabled(NewState == ECheckBoxState::Checked);
+		GraphObj->SetShowSoftReferencesEnabled( !GraphObj->IsShowSoftReferences() );
 		RebuildGraph();
 	}
 }
 
-ECheckBoxState SReferenceViewer::IsShowSoftReferencesChecked() const
+bool SReferenceViewer::IsShowSoftReferencesChecked() const
 {
-	if (GraphObj)
-	{
-		return GraphObj->IsShowSoftReferences()? ECheckBoxState::Checked : ECheckBoxState::Unchecked;
-	}
-	else
-	{
-		return ECheckBoxState::Unchecked;
-	}
+	return GraphObj ? GraphObj->IsShowSoftReferences() : false;
 }
 
-void SReferenceViewer::OnShowHardReferencesChanged(ECheckBoxState NewState)
+void SReferenceViewer::OnShowHardReferencesChanged()
 {
 	if (GraphObj)
 	{
-		GraphObj->SetShowHardReferencesEnabled(NewState == ECheckBoxState::Checked);
+		GraphObj->SetShowHardReferencesEnabled(!GraphObj->IsShowHardReferences());
 		RebuildGraph();
 	}
 }
 
 
-ECheckBoxState SReferenceViewer::IsShowHardReferencesChecked() const
+bool SReferenceViewer::IsShowHardReferencesChecked() const
 {
-	if (GraphObj)
-	{
-		return GraphObj->IsShowHardReferences() ? ECheckBoxState::Checked : ECheckBoxState::Unchecked;
-	}
-	else
-	{
-		return ECheckBoxState::Unchecked;
-	}
+	return GraphObj ? GraphObj->IsShowHardReferences() : false;
 }
 
-void SReferenceViewer::OnShowFilteredPackagesOnlyChanged(ECheckBoxState NewState)
+void SReferenceViewer::OnShowFilteredPackagesOnlyChanged()
 {
 	if (GraphObj)
 	{
-		GraphObj->SetShowFilteredPackagesOnlyEnabled(NewState == ECheckBoxState::Checked);
+		GraphObj->SetShowFilteredPackagesOnlyEnabled(!GraphObj->IsShowFilteredPackagesOnly());
 	}
 	UpdateIsPassingFilterPackageCallback();
 }
 
 
-ECheckBoxState SReferenceViewer::IsShowFilteredPackagesOnlyChecked() const
+bool SReferenceViewer::IsShowFilteredPackagesOnlyChecked() const
 {
-	if (GraphObj)
-	{
-		return GraphObj->IsShowFilteredPackagesOnly() ? ECheckBoxState::Checked : ECheckBoxState::Unchecked;
-	}
-	else
-	{
-		return ECheckBoxState::Unchecked;
-	}
+	return GraphObj ? GraphObj->IsShowFilteredPackagesOnly() : false;
 }
 
-
 void SReferenceViewer::UpdateIsPassingFilterPackageCallback()
 {
 	if (GraphObj)
@@ -1248,161 +1044,102 @@ void SReferenceViewer::UpdateIsPassingFilterPackageCallback()
 	}
 }
 
-
-void SReferenceViewer::OnCompactModeChanged(ECheckBoxState NewState)
+void SReferenceViewer::OnCompactModeChanged()
 {
 	if (GraphObj)
 	{
-		GraphObj->SetCompactModeEnabled(NewState == ECheckBoxState::Checked);
+		GraphObj->SetCompactModeEnabled(!GraphObj->IsCompactMode());
+		RebuildGraph();
 	}
-	
-	RebuildGraph();
 }
 
-
-ECheckBoxState SReferenceViewer::IsCompactModeChecked() const
+bool SReferenceViewer::IsCompactModeChecked() const
 {
-	if (GraphObj)
-	{
-		return GraphObj->IsCompactMode() ? ECheckBoxState::Checked : ECheckBoxState::Unchecked;
-	}
-	else
-	{
-		return ECheckBoxState::Unchecked;
-	}
+	return GraphObj ? GraphObj->IsCompactMode() : false;
 }
 
-void SReferenceViewer::OnShowDuplicatesChanged(ECheckBoxState NewState)
+void SReferenceViewer::OnShowDuplicatesChanged()
 {
 	if (GraphObj)
 	{
-		GraphObj->SetShowDuplicatesEnabled(NewState == ECheckBoxState::Checked);
+		GraphObj->SetShowDuplicatesEnabled(!GraphObj->IsShowDuplicates());
+		RebuildGraph();
 	}
-	
-	RebuildGraph();
 }
 
-
-ECheckBoxState SReferenceViewer::IsShowDuplicatesChecked() const
+bool SReferenceViewer::IsShowDuplicatesChecked() const
 {
-	if (GraphObj)
-	{
-		return GraphObj->IsShowDuplicates() ? ECheckBoxState::Checked : ECheckBoxState::Unchecked;
-	}
-	else
-	{
-		return ECheckBoxState::Unchecked;
-	}
+	return GraphObj ? GraphObj->IsShowDuplicates() : false;
 }
 
-void SReferenceViewer::OnShowEditorOnlyReferencesChanged(ECheckBoxState NewState)
+void SReferenceViewer::OnShowEditorOnlyReferencesChanged()
 {
 	if (GraphObj)
 	{
-		GraphObj->SetShowEditorOnlyReferencesEnabled(NewState == ECheckBoxState::Checked);
+		GraphObj->SetShowEditorOnlyReferencesEnabled(!GraphObj->IsShowEditorOnlyReferences());
 		RebuildGraph();
 	}
 }
 
-
-ECheckBoxState SReferenceViewer::IsShowEditorOnlyReferencesChecked() const
+bool SReferenceViewer::IsShowEditorOnlyReferencesChecked() const
 {
-	if (GraphObj)
-	{
-		return GraphObj->IsShowEditorOnlyReferences() ? ECheckBoxState::Checked : ECheckBoxState::Unchecked;
-	}
-	else
-	{
-		return ECheckBoxState::Unchecked;
-	}
+	return GraphObj ? GraphObj->IsShowEditorOnlyReferences() : false;
 }
 
 
-EVisibility SReferenceViewer::GetManagementReferencesVisibility() const
+bool SReferenceViewer::GetManagementReferencesVisibility() const
 {
-	if (bShowShowReferencesOptions && UAssetManager::IsValid())
-	{
-		return EVisibility::SelfHitTestInvisible;
-	}
-	return EVisibility::Collapsed;
+	return bShowShowReferencesOptions && UAssetManager::IsValid();
 }
 
-void SReferenceViewer::OnShowManagementReferencesChanged(ECheckBoxState NewState)
+void SReferenceViewer::OnShowManagementReferencesChanged()
 {
 	if (GraphObj)
 	{
 		// This can take a few seconds if it isn't ready
 		UAssetManager::Get().UpdateManagementDatabase();
 
-		GraphObj->SetShowManagementReferencesEnabled(NewState == ECheckBoxState::Checked);
+		GraphObj->SetShowManagementReferencesEnabled(!GraphObj->IsShowManagementReferences());
 		RebuildGraph();
 	}
 }
 
-ECheckBoxState SReferenceViewer::IsShowManagementReferencesChecked() const
+bool SReferenceViewer::IsShowManagementReferencesChecked() const
 {
-	if (GraphObj)
-	{
-		return GraphObj->IsShowManagementReferences() ? ECheckBoxState::Checked : ECheckBoxState::Unchecked;
-	}
-	else
-	{
-		return ECheckBoxState::Unchecked;
-	}
+	return GraphObj ? GraphObj->IsShowManagementReferences() : false;
 }
 
-void SReferenceViewer::OnShowSearchableNamesChanged(ECheckBoxState NewState)
+void SReferenceViewer::OnShowSearchableNamesChanged()
 {
 	if (GraphObj)
 	{
-		GraphObj->SetShowSearchableNames(NewState == ECheckBoxState::Checked);
+		GraphObj->SetShowSearchableNames(!GraphObj->IsShowSearchableNames());
 		RebuildGraph();
 	}
 }
 
-ECheckBoxState SReferenceViewer::IsShowSearchableNamesChecked() const
+bool SReferenceViewer::IsShowSearchableNamesChecked() const
 {
-	if (GraphObj)
-	{
-		return GraphObj->IsShowSearchableNames() ? ECheckBoxState::Checked : ECheckBoxState::Unchecked;
-	}
-	else
-	{
-		return ECheckBoxState::Unchecked;
-	}
+	return GraphObj ? GraphObj->IsShowSearchableNames() : false;
 }
 
-void SReferenceViewer::OnShowNativePackagesChanged(ECheckBoxState NewState)
+void SReferenceViewer::OnShowNativePackagesChanged()
 {
 	if (GraphObj)
 	{
-		GraphObj->SetShowNativePackages(NewState == ECheckBoxState::Checked);
+		GraphObj->SetShowNativePackages(!GraphObj->IsShowNativePackages());
 		RebuildGraph();
 	}
 }
 
-ECheckBoxState SReferenceViewer::IsShowNativePackagesChecked() const
+bool SReferenceViewer::IsShowNativePackagesChecked() const
 {
-	if (GraphObj)
-	{
-		return GraphObj->IsShowNativePackages() ? ECheckBoxState::Checked : ECheckBoxState::Unchecked;
-	}
-	else
-	{
-		return ECheckBoxState::Unchecked;
-	}
+	return GraphObj ? GraphObj->IsShowNativePackages() : false;
 }
 
 int32 SReferenceViewer::GetSearchBreadthCount() const
 {
-	if ( GraphObj )
-	{
-		return GraphObj->GetSearchBreadthLimit();
-	}
-	else
-	{
-		return 0;
-	}
+	return GraphObj ? GraphObj->GetSearchBreadthLimit() : 0;
 }
 
 void SReferenceViewer::OnSearchBreadthCommitted(int32 NewValue)
@@ -1443,6 +1180,115 @@ void SReferenceViewer::RegisterActions()
 		FExecuteAction::CreateSP(this, &SReferenceViewer::ReCenterGraph),
 		FCanExecuteAction());
 
+	ReferenceViewerActions->MapAction(
+		FAssetManagerEditorCommands::Get().IncreaseReferencerSearchDepth,
+		FExecuteAction::CreateLambda( [this] { OnSearchReferencerDepthCommitted( GetSearchReferencerDepthCount() + 1); } ),
+		FCanExecuteAction());
+
+	ReferenceViewerActions->MapAction(
+		FAssetManagerEditorCommands::Get().DecreaseReferencerSearchDepth,
+		FExecuteAction::CreateLambda( [this] { OnSearchReferencerDepthCommitted( GetSearchReferencerDepthCount() - 1); } ),
+		FCanExecuteAction());
+
+	ReferenceViewerActions->MapAction(
+		FAssetManagerEditorCommands::Get().SetReferencerSearchDepth,
+		FExecuteAction::CreateLambda( [this] { FSlateApplication::Get().SetKeyboardFocus(ReferencerCountBox, EFocusCause::SetDirectly); } ),
+		FCanExecuteAction());
+
+	ReferenceViewerActions->MapAction(
+		FAssetManagerEditorCommands::Get().IncreaseDependencySearchDepth,
+		FExecuteAction::CreateLambda( [this] { OnSearchDependencyDepthCommitted( GetSearchDependencyDepthCount() + 1); } ),
+		FCanExecuteAction());
+
+	ReferenceViewerActions->MapAction(
+		FAssetManagerEditorCommands::Get().DecreaseDependencySearchDepth,
+		FExecuteAction::CreateLambda( [this] { OnSearchDependencyDepthCommitted( GetSearchDependencyDepthCount() - 1); } ),
+		FCanExecuteAction());
+
+	ReferenceViewerActions->MapAction(
+		FAssetManagerEditorCommands::Get().SetDependencySearchDepth,
+		FExecuteAction::CreateLambda( [this] { FSlateApplication::Get().SetKeyboardFocus(DependencyCountBox, EFocusCause::SetDirectly); } ),
+		FCanExecuteAction());
+
+	ReferenceViewerActions->MapAction(
+		FAssetManagerEditorCommands::Get().IncreaseBreadth,
+		FExecuteAction::CreateLambda( [this] { OnSearchBreadthCommitted( GetSearchBreadthCount() + 1); } ),
+		FCanExecuteAction());
+
+	ReferenceViewerActions->MapAction(
+		FAssetManagerEditorCommands::Get().DecreaseBreadth,
+		FExecuteAction::CreateLambda( [this] { OnSearchBreadthCommitted( GetSearchBreadthCount() - 1); } ),
+		FCanExecuteAction());
+
+	ReferenceViewerActions->MapAction(
+		FAssetManagerEditorCommands::Get().SetBreadth,
+		FExecuteAction::CreateLambda( [this] { FSlateApplication::Get().SetKeyboardFocus(BreadthLimitBox, EFocusCause::SetDirectly); } ),
+		FCanExecuteAction());
+
+	ReferenceViewerActions->MapAction(
+		FAssetManagerEditorCommands::Get().ShowSoftReferences,
+		FExecuteAction::CreateSP(this, &SReferenceViewer::OnShowSoftReferencesChanged),
+		FCanExecuteAction(),	
+		FIsActionChecked::CreateSP(this, &SReferenceViewer::IsShowSoftReferencesChecked),
+		FIsActionButtonVisible::CreateLambda([this] { return bShowShowReferencesOptions; }));
+
+	ReferenceViewerActions->MapAction(
+		FAssetManagerEditorCommands::Get().ShowHardReferences,
+		FExecuteAction::CreateSP(this, &SReferenceViewer::OnShowHardReferencesChanged),
+		FCanExecuteAction(),	
+		FIsActionChecked::CreateSP(this, &SReferenceViewer::IsShowHardReferencesChecked),
+		FIsActionButtonVisible::CreateLambda([this] { return bShowShowReferencesOptions; }));
+
+	ReferenceViewerActions->MapAction(
+		FAssetManagerEditorCommands::Get().ShowEditorOnlyReferences,
+		FExecuteAction::CreateSP(this, &SReferenceViewer::OnShowEditorOnlyReferencesChanged),
+		FCanExecuteAction(),	
+		FIsActionChecked::CreateSP(this, &SReferenceViewer::IsShowEditorOnlyReferencesChecked),
+		FIsActionButtonVisible::CreateLambda([this] { return bShowShowReferencesOptions; }));
+
+	ReferenceViewerActions->MapAction(
+		FAssetManagerEditorCommands::Get().ShowManagementReferences,
+		FExecuteAction::CreateSP(this, &SReferenceViewer::OnShowManagementReferencesChanged),
+		FCanExecuteAction(),	
+		FIsActionChecked::CreateSP(this, &SReferenceViewer::IsShowManagementReferencesChecked),
+		FIsActionButtonVisible::CreateSP(this, &SReferenceViewer::GetManagementReferencesVisibility));
+
+	ReferenceViewerActions->MapAction(
+		FAssetManagerEditorCommands::Get().ShowGamePlayTags,
+		FExecuteAction::CreateSP(this, &SReferenceViewer::OnShowSearchableNamesChanged),
+		FCanExecuteAction(),	
+		FIsActionChecked::CreateSP(this, &SReferenceViewer::IsShowSearchableNamesChecked),
+		FIsActionButtonVisible::CreateLambda([this] { return bShowShowSearchableNames; }));
+
+	ReferenceViewerActions->MapAction(
+		FAssetManagerEditorCommands::Get().ShowNativePackages,
+		FExecuteAction::CreateSP(this, &SReferenceViewer::OnShowNativePackagesChanged),
+		FCanExecuteAction(),	
+		FIsActionChecked::CreateSP(this, &SReferenceViewer::IsShowNativePackagesChecked),
+		FIsActionButtonVisible::CreateLambda([this] { return bShowShowNativePackages; }));
+
+
+	ReferenceViewerActions->MapAction(
+		FAssetManagerEditorCommands::Get().ShowDuplicates,
+		FExecuteAction::CreateSP(this, &SReferenceViewer::OnShowDuplicatesChanged),
+		FCanExecuteAction(),	
+		FIsActionChecked::CreateSP(this, &SReferenceViewer::IsShowDuplicatesChecked));
+
+	ReferenceViewerActions->MapAction(
+		FAssetManagerEditorCommands::Get().CompactMode,
+		FExecuteAction::CreateSP(this, &SReferenceViewer::OnCompactModeChanged),
+		FCanExecuteAction(),	
+		FIsActionChecked::CreateSP(this, &SReferenceViewer::IsCompactModeChecked),
+		FIsActionButtonVisible::CreateLambda([this] { return bShowCompactMode; }));
+
+
+	ReferenceViewerActions->MapAction(
+		FAssetManagerEditorCommands::Get().FilterSearch,
+		FExecuteAction::CreateSP(this, &SReferenceViewer::OnShowFilteredPackagesOnlyChanged),
+		FCanExecuteAction(),	
+		FIsActionChecked::CreateSP(this, &SReferenceViewer::IsShowFilteredPackagesOnlyChecked),
+		FIsActionButtonVisible::CreateLambda([this] { return bShowShowFilteredPackagesOnly; }) );
+
 	ReferenceViewerActions->MapAction(
 		FAssetManagerEditorCommands::Get().CopyReferencedObjects,
 		FExecuteAction::CreateSP(this, &SReferenceViewer::CopyReferencedObjects),
@@ -2034,6 +1880,32 @@ void SReferenceViewer::HandleOnSearchTextCommitted(const FText& SearchText, ETex
 	}
 	
 	GraphEditorPtr->ZoomToFit(true);
+	FSlateApplication::Get().SetKeyboardFocus(GraphEditorPtr, EFocusCause::SetDirectly);
+}
+
+TSharedRef<SWidget> SReferenceViewer::GetShowMenuContent()
+{
+ 	FMenuBuilder MenuBuilder(true, ReferenceViewerActions);
+
+	MenuBuilder.BeginSection("ReferenceTypes", LOCTEXT("ReferenceTypes", "Reference Types"));
+	MenuBuilder.AddMenuEntry(FAssetManagerEditorCommands::Get().ShowSoftReferences);
+	MenuBuilder.AddMenuEntry(FAssetManagerEditorCommands::Get().ShowHardReferences);
+	MenuBuilder.AddMenuEntry(FAssetManagerEditorCommands::Get().ShowEditorOnlyReferences);
+	MenuBuilder.EndSection();
+
+	MenuBuilder.BeginSection("Assets", LOCTEXT("Assets", "Assets"));
+	MenuBuilder.AddMenuEntry(FAssetManagerEditorCommands::Get().ShowManagementReferences);
+	MenuBuilder.AddMenuEntry(FAssetManagerEditorCommands::Get().ShowGamePlayTags);
+	MenuBuilder.AddMenuEntry(FAssetManagerEditorCommands::Get().ShowNativePackages);
+	MenuBuilder.EndSection();
+
+	MenuBuilder.BeginSection("ViewOptions", LOCTEXT("ViewOptions", "View Options"));
+	MenuBuilder.AddMenuEntry(FAssetManagerEditorCommands::Get().ShowDuplicates);
+	MenuBuilder.AddMenuEntry(FAssetManagerEditorCommands::Get().FilterSearch);
+	MenuBuilder.AddMenuEntry(FAssetManagerEditorCommands::Get().CompactMode);
+	MenuBuilder.EndSection();
+
+	return MenuBuilder.MakeWidget();
 }
 
 #undef LOCTEXT_NAMESPACE
\ No newline at end of file
diff --git a/Engine/Plugins/Editor/AssetManagerEditor/Source/AssetManagerEditor/Private/ReferenceViewer/SReferenceViewer.h b/Engine/Plugins/Editor/AssetManagerEditor/Source/AssetManagerEditor/Private/ReferenceViewer/SReferenceViewer.h
--- a/Engine/Plugins/Editor/AssetManagerEditor/Source/AssetManagerEditor/Private/ReferenceViewer/SReferenceViewer.h
+++ b/Engine/Plugins/Editor/AssetManagerEditor/Source/AssetManagerEditor/Private/ReferenceViewer/SReferenceViewer.h
@@ -47,6 +47,10 @@ public:
 	/** Called when the current registry source changes */
 	void SetCurrentRegistrySource(const FAssetManagerEditorRegistrySource* RegistrySource);
 
+	/**SWidget interface **/
+	virtual bool SupportsKeyboardFocus() const override { return true; }
+	virtual FReply OnKeyDown( const FGeometry& MyGeometry, const FKeyEvent& InKeyEvent ) override;
+
 private:
 
 	/** Call after a structural change is made that causes the graph to be recreated */
@@ -120,35 +124,36 @@ private:
 	void HandleCollectionFilterChanged(TSharedPtr<FName> Item, ESelectInfo::Type SelectInfo);
 	FText GetCollectionFilterText() const;
 
-	void OnShowSoftReferencesChanged( ECheckBoxState NewState );
-	ECheckBoxState IsShowSoftReferencesChecked() const;
-	void OnShowHardReferencesChanged(ECheckBoxState NewState);
-	ECheckBoxState IsShowHardReferencesChecked() const;
-	void OnShowEditorOnlyReferencesChanged(ECheckBoxState NewState);
-	ECheckBoxState IsShowEditorOnlyReferencesChecked() const;
+	void OnShowSoftReferencesChanged();
+	bool IsShowSoftReferencesChecked() const;
+	void OnShowHardReferencesChanged();
+	bool IsShowHardReferencesChecked() const;
+	void OnShowEditorOnlyReferencesChanged();
+	bool IsShowEditorOnlyReferencesChecked() const;
 
-	void OnShowFilteredPackagesOnlyChanged(ECheckBoxState NewState);
-	ECheckBoxState IsShowFilteredPackagesOnlyChecked() const;
+	void OnShowFilteredPackagesOnlyChanged();
+	bool IsShowFilteredPackagesOnlyChecked() const;
 	void UpdateIsPassingFilterPackageCallback();
 
-	void OnCompactModeChanged(ECheckBoxState NewState);
-	ECheckBoxState IsCompactModeChecked() const;
-
-	void OnShowDuplicatesChanged(ECheckBoxState NewState);
-	ECheckBoxState IsShowDuplicatesChecked() const;
+	void OnCompactModeChanged();
+	bool IsCompactModeChecked() const;
 
-	EVisibility GetManagementReferencesVisibility() const;
-	void OnShowManagementReferencesChanged(ECheckBoxState NewState);
-	ECheckBoxState IsShowManagementReferencesChecked() const;
+	void OnShowDuplicatesChanged();
+	bool IsShowDuplicatesChecked() const;
 
-	void OnShowSearchableNamesChanged(ECheckBoxState NewState);
-	ECheckBoxState IsShowSearchableNamesChecked() const;
-	void OnShowNativePackagesChanged(ECheckBoxState NewState);
-	ECheckBoxState IsShowNativePackagesChecked() const;
+	bool GetManagementReferencesVisibility() const;
+	void OnShowManagementReferencesChanged();
+	bool IsShowManagementReferencesChecked() const;
+	void OnShowSearchableNamesChanged();
+	bool IsShowSearchableNamesChecked() const;
+	void OnShowNativePackagesChanged();
+	bool IsShowNativePackagesChecked() const;
 
 	int32 GetSearchBreadthCount() const;
 	void OnSearchBreadthCommitted(int32 NewValue);
 
+	TSharedRef<SWidget> GetShowMenuContent();
+
 	void RegisterActions();
 	void ShowSelectionInContentBrowser();
 	void OpenSelectedInAssetEditor();
@@ -193,6 +198,9 @@ private:
 
 	TSharedPtr<FUICommandList> ReferenceViewerActions;
 	TSharedPtr<SSearchBox> SearchBox;
+	TSharedPtr<SWidget> ReferencerCountBox;
+	TSharedPtr<SWidget> DependencyCountBox;
+	TSharedPtr<SWidget> BreadthLimitBox;
 
 	UEdGraph_ReferenceViewer* GraphObj;
 
