Reference Viewer - Settings now get saved to project config

diff --git a/Engine/Config/BaseEditorPerProjectUserSettings.ini b/Engine/Config/BaseEditorPerProjectUserSettings.ini
--- a/Engine/Config/BaseEditorPerProjectUserSettings.ini
+++ b/Engine/Config/BaseEditorPerProjectUserSettings.ini
@@ -931,6 +931,24 @@ bEnabled=False
 bEnableReinstancing=True
 bAutomaticallyCompileNewClasses=True
 
+[/Script/AssetManagerEditor.ReferenceViewerSettings]
+bLimitSearchDepth=True
+bIsShowReferencers=True
+MaxSearchReferencerDepth=1
+bIsShowDependencies=True
+MaxSearchDependencyDepth=1
+bLimitSearchBreadth=True
+MaxSearchBreadth=20
+bIsShowSoftReferences=True
+bIsShowHardReferences=True
+bIsShowEditorOnlyReferences=True
+bIsShowManagementReferences=False
+bIsShowSearchableNames=False
+bIsShowNativePackages=False
+bIsShowDuplicates=True
+bIsShowFilteredPackagesOnly=False
+bIsCompactMode=False
+
 [/Script/BlueprintHeaderView.BlueprintHeaderViewSettings]
 SyntaxColors=(Comment=(R=0.258183,G=0.539479,B=0.068478,A=1.000000),Error=(R=0.863157,G=0.035601,B=0.035601,A=1.000000),Macro=(R=0.356400,G=0.040915,B=0.520996,A=1.000000),Typename=(R=0.000000,G=0.300000,B=0.300000,A=1.000000),Identifier=(R=1.000000,G=1.000000,B=1.000000,A=1.000000),Keyword=(R=0.019382,G=0.496933,B=1.000000,A=1.000000))
 FontSize=9
diff --git a/Engine/Plugins/Editor/AssetManagerEditor/Source/AssetManagerEditor/Private/ReferenceViewer/EdGraph_ReferenceViewer.cpp b/Engine/Plugins/Editor/AssetManagerEditor/Source/AssetManagerEditor/Private/ReferenceViewer/EdGraph_ReferenceViewer.cpp
--- a/Engine/Plugins/Editor/AssetManagerEditor/Source/AssetManagerEditor/Private/ReferenceViewer/EdGraph_ReferenceViewer.cpp
+++ b/Engine/Plugins/Editor/AssetManagerEditor/Source/AssetManagerEditor/Private/ReferenceViewer/EdGraph_ReferenceViewer.cpp
@@ -2,6 +2,7 @@
 
 #include "ReferenceViewer/EdGraph_ReferenceViewer.h"
 #include "ReferenceViewer/EdGraphNode_Reference.h"
+#include "ReferenceViewer/ReferenceViewerSettings.h"
 #include "EdGraph/EdGraphPin.h"
 #include "AssetRegistry/ARFilter.h"
 #include "AssetRegistry/AssetRegistryModule.h"
@@ -42,23 +43,8 @@ UEdGraph_ReferenceViewer::UEdGraph_ReferenceViewer(const FObjectInitializer& Obj
 {
 	AssetThumbnailPool = MakeShareable( new FAssetThumbnailPool(1024) );
 
-	MaxSearchDependencyDepth = 1;
-	MaxSearchReferencerDepth = 1;
-	MaxSearchBreadth = GetDefault<UEditorProjectAppearanceSettings>()->ReferenceViewerDefaultMaxSearchBreadth;
-
-	bLimitSearchDepth = true;
-	bLimitSearchBreadth = true;
-	bIsShowSoftReferences = true;
-	bIsShowHardReferences = true;
-	bIsShowEditorOnlyReferences = true;
-	bIsShowManagementReferences = false;
-	bIsShowSearchableNames = GetDefault<UEditorProjectAppearanceSettings>()->ShowSearchableNames == EReferenceViewerSettingMode::ShowByDefault;
-	bIsShowNativePackages = false;
-	bIsShowReferencers = true;
-	bIsShowDependencies = true;
-	bIsShowFilteredPackagesOnly = false;
-	bIsCompactMode = false;
-	bIsShowDuplicates = true;
+	Settings = GetMutableDefault<UReferenceViewerSettings>();
+
 	bUseNodeInfos = true;
 }
 
@@ -79,7 +65,7 @@ void UEdGraph_ReferenceViewer::SetGraphRoot(const TArray<FAssetIdentifier>& Grap
 	{
 		if (AssetId.IsValue())
 		{
-			bIsShowSearchableNames = true;
+			Settings->SetShowSearchableNames(true);
 		}
 		else if (AssetId.GetPrimaryAssetId().IsValid())
 		{
@@ -88,7 +74,7 @@ void UEdGraph_ReferenceViewer::SetGraphRoot(const TArray<FAssetIdentifier>& Grap
 				UAssetManager::Get().UpdateManagementDatabase();
 			}
 			
-			bIsShowManagementReferences = true;
+			Settings->SetShowManagementReferencesEnabled(true);
 		}
 	}
 }
@@ -139,166 +125,6 @@ UEdGraphNode_Reference* UEdGraph_ReferenceViewer::RebuildGraph()
 	return NewRootNode;
 }
 
-bool UEdGraph_ReferenceViewer::IsSearchDepthLimited() const
-{
-	return bLimitSearchDepth;
-}
-
-bool UEdGraph_ReferenceViewer::IsSearchBreadthLimited() const
-{
-	return bLimitSearchBreadth;
-}
-
-bool UEdGraph_ReferenceViewer::IsShowSoftReferences() const
-{
-	return bIsShowSoftReferences;
-}
-
-bool UEdGraph_ReferenceViewer::IsShowHardReferences() const
-{
-	return bIsShowHardReferences;
-}
-
-bool UEdGraph_ReferenceViewer::IsShowFilteredPackagesOnly() const
-{
-	return bIsShowFilteredPackagesOnly;
-}
-
-bool UEdGraph_ReferenceViewer::IsCompactMode() const
-{
-	return bIsCompactMode;
-}
-
-bool UEdGraph_ReferenceViewer::IsShowDuplicates() const
-{
-	return bIsShowDuplicates;
-}
-
-bool UEdGraph_ReferenceViewer::IsShowEditorOnlyReferences() const
-{
-	return bIsShowEditorOnlyReferences;
-}
-
-bool UEdGraph_ReferenceViewer::IsShowManagementReferences() const
-{
-	return bIsShowManagementReferences;
-}
-
-bool UEdGraph_ReferenceViewer::IsShowSearchableNames() const
-{
-	return bIsShowSearchableNames;
-}
-
-bool UEdGraph_ReferenceViewer::IsShowNativePackages() const
-{
-	return bIsShowNativePackages;
-}
-
-bool UEdGraph_ReferenceViewer::IsShowReferencers() const
-{
-	return bIsShowReferencers;
-}
-
-bool UEdGraph_ReferenceViewer::IsShowDependencies() const
-{
-	return bIsShowDependencies;
-}
-
-void UEdGraph_ReferenceViewer::SetSearchDepthLimitEnabled(bool newEnabled)
-{
-	bLimitSearchDepth = newEnabled;
-}
-
-void UEdGraph_ReferenceViewer::SetSearchBreadthLimitEnabled(bool newEnabled)
-{
-	bLimitSearchBreadth = newEnabled;
-}
-
-void UEdGraph_ReferenceViewer::SetShowSoftReferencesEnabled(bool newEnabled)
-{
-	bIsShowSoftReferences = newEnabled;
-}
-
-void UEdGraph_ReferenceViewer::SetShowHardReferencesEnabled(bool newEnabled)
-{
-	bIsShowHardReferences = newEnabled;
-}
-
-void UEdGraph_ReferenceViewer::SetShowFilteredPackagesOnlyEnabled(bool newEnabled)
-{
-	bIsShowFilteredPackagesOnly = newEnabled;
-}
-
-void UEdGraph_ReferenceViewer::SetCompactModeEnabled(bool newEnabled)
-{
-	bIsCompactMode = newEnabled;
-}
-
-void UEdGraph_ReferenceViewer::SetShowDuplicatesEnabled(bool newEnabled)
-{
-	bIsShowDuplicates = newEnabled;
-}
-
-void UEdGraph_ReferenceViewer::SetShowEditorOnlyReferencesEnabled(bool newEnabled)
-{
-	bIsShowEditorOnlyReferences = newEnabled;
-}
-
-void UEdGraph_ReferenceViewer::SetShowManagementReferencesEnabled(bool newEnabled)
-{
-	bIsShowManagementReferences = newEnabled;
-}
-
-void UEdGraph_ReferenceViewer::SetShowSearchableNames(bool newEnabled)
-{
-	bIsShowSearchableNames = newEnabled;
-}
-
-void UEdGraph_ReferenceViewer::SetShowNativePackages(bool newEnabled)
-{
-	bIsShowNativePackages = newEnabled;
-}
-
-void UEdGraph_ReferenceViewer::SetShowReferencers(const bool newEnabled)
-{
-	bIsShowReferencers = newEnabled;
-}
-
-void UEdGraph_ReferenceViewer::SetShowDependencies(const bool newEnabled)
-{
-	bIsShowDependencies = newEnabled;
-}
-
-int32 UEdGraph_ReferenceViewer::GetSearchDependencyDepthLimit() const
-{
-	return MaxSearchDependencyDepth;
-}
-
-void UEdGraph_ReferenceViewer::SetSearchDependencyDepthLimit(int32 NewDepthLimit)
-{
-	MaxSearchDependencyDepth = FMath::Max(NewDepthLimit, 0);
-}
-
-int32 UEdGraph_ReferenceViewer::GetSearchReferencerDepthLimit() const
-{
-	return MaxSearchReferencerDepth;
-}
-
-void UEdGraph_ReferenceViewer::SetSearchReferencerDepthLimit(int32 NewDepthLimit)
-{
-	MaxSearchReferencerDepth = FMath::Max(NewDepthLimit, 0);
-}
-
-int32 UEdGraph_ReferenceViewer::GetSearchBreadthLimit() const
-{
-	return MaxSearchBreadth;
-}
-
-void UEdGraph_ReferenceViewer::SetSearchBreadthLimit(int32 NewBreadthLimit)
-{
-	MaxSearchBreadth = FMath::Max(NewBreadthLimit, 0);
-}
-
 FName UEdGraph_ReferenceViewer::GetCurrentCollectionFilter() const
 {
 	return CurrentCollectionFilter;
@@ -309,16 +135,6 @@ void UEdGraph_ReferenceViewer::SetCurrentCollectionFilter(FName NewFilter)
 	CurrentCollectionFilter = NewFilter;
 }
 
-bool UEdGraph_ReferenceViewer::GetEnableCollectionFilter() const
-{
-	return bEnableCollectionFilter;
-}
-
-void UEdGraph_ReferenceViewer::SetEnableCollectionFilter(bool bEnabled)
-{
-	bEnableCollectionFilter = bEnabled;
-}
-
 FAssetManagerDependencyQuery UEdGraph_ReferenceViewer::GetReferenceSearchFlags(bool bHardOnly) const
 {
 	using namespace UE::AssetRegistry;
@@ -326,19 +142,19 @@ FAssetManagerDependencyQuery UEdGraph_ReferenceViewer::GetReferenceSearchFlags(b
 	Query.Categories = EDependencyCategory::None;
 	Query.Flags = EDependencyQuery::NoRequirements;
 
-	bool bLocalIsShowSoftReferences = bIsShowSoftReferences && !bHardOnly;
-	if (bLocalIsShowSoftReferences || bIsShowHardReferences)
+	bool bLocalIsShowSoftReferences = Settings->IsShowSoftReferences() && !bHardOnly;
+	if (bLocalIsShowSoftReferences || Settings->IsShowHardReferences())
 	{
 		Query.Categories |= EDependencyCategory::Package;
 		Query.Flags |= bLocalIsShowSoftReferences ? EDependencyQuery::NoRequirements : EDependencyQuery::Hard;
-		Query.Flags |= bIsShowHardReferences ? EDependencyQuery::NoRequirements : EDependencyQuery::Soft;
-		Query.Flags |= bIsShowEditorOnlyReferences ? EDependencyQuery::NoRequirements : EDependencyQuery::Game;
+		Query.Flags |= Settings->IsShowHardReferences() ? EDependencyQuery::NoRequirements : EDependencyQuery::Soft;
+		Query.Flags |= Settings->IsShowEditorOnlyReferences() ? EDependencyQuery::NoRequirements : EDependencyQuery::Game;
 	}
-	if (bIsShowSearchableNames && !bHardOnly)
+	if (Settings->IsShowSearchableNames() && !bHardOnly)
 	{
 		Query.Categories |= EDependencyCategory::SearchableName;
 	}
-	if (bIsShowManagementReferences)
+	if (Settings->IsShowManagementReferences())
 	{
 		Query.Categories |= EDependencyCategory::Manage;
 		Query.Flags |= bHardOnly ? EDependencyQuery::Direct : EDependencyQuery::NoRequirements;
@@ -354,7 +170,7 @@ UEdGraphNode_Reference* UEdGraph_ReferenceViewer::ConstructNodes(const TArray<FA
 	if (GraphRootIdentifiers.Num() > 0)
 	{
 		// It both were false, nothing (other than the GraphRootIdentifiers) would be displayed
-		check(bIsShowReferencers || bIsShowDependencies);
+		check(Settings->IsShowReferencers() || Settings->IsShowDependencies());
 
 		// Refresh the current collection filter
 		CurrentCollectionPackages.Empty();
@@ -376,27 +192,27 @@ UEdGraphNode_Reference* UEdGraph_ReferenceViewer::ConstructNodes(const TArray<FA
 		TMap<FAssetIdentifier, ReferenceNodeInfo> ReferenceNodeInfos;
 		ReferenceNodeInfo& RootNodeInfo = ReferenceNodeInfos.FindOrAdd( GraphRootIdentifiers[0], ReferenceNodeInfo(GraphRootIdentifiers[0], true));
 		RootNodeInfo.Parents.Emplace(FAssetIdentifier(NAME_None));
-		RecursivelyPopulateNodeInfos(true, GraphRootIdentifiers[0], ReferenceNodeInfos, 0, MaxSearchReferencerDepth);
+		RecursivelyPopulateNodeInfos(true, GraphRootIdentifiers[0], ReferenceNodeInfos, 0, Settings->GetSearchReferencerDepthLimit());
 
 		TMap<FAssetIdentifier, ReferenceNodeInfo> DependencyNodeInfos;
 		ReferenceNodeInfo& DRootNodeInfo = DependencyNodeInfos.FindOrAdd( GraphRootIdentifiers[0], ReferenceNodeInfo(GraphRootIdentifiers[0], false));
 		DRootNodeInfo.Parents.Emplace(FAssetIdentifier(NAME_None));
-		RecursivelyPopulateNodeInfos(false, GraphRootIdentifiers[0], DependencyNodeInfos, 0, MaxSearchDependencyDepth);
+		RecursivelyPopulateNodeInfos(false, GraphRootIdentifiers[0], DependencyNodeInfos, 0, Settings->GetSearchDependencyDepthLimit());
 
 		TMap<FAssetIdentifier, int32> ReferencerNodeSizes;
 		TSet<FAssetIdentifier> VisitedReferencerSizeNames;
-		if (bIsShowReferencers)
+		if (Settings->IsShowReferencers())
 		{
 			const int32 ReferencerDepth = 1;
-			RecursivelyGatherSizes(/*bReferencers=*/true, GraphRootIdentifiers, CurrentCollectionPackages, ReferencerDepth, MaxSearchReferencerDepth, VisitedReferencerSizeNames, ReferencerNodeSizes);
+			RecursivelyGatherSizes(/*bReferencers=*/true, GraphRootIdentifiers, CurrentCollectionPackages, ReferencerDepth, Settings->GetSearchReferencerDepthLimit(), VisitedReferencerSizeNames, ReferencerNodeSizes);
 		}
 
 		TMap<FAssetIdentifier, int32> DependencyNodeSizes;
 		TSet<FAssetIdentifier> VisitedDependencySizeNames;
-		if (bIsShowDependencies)
+		if (Settings->IsShowDependencies())
 		{
 			const int32 DependencyDepth = 1;
-			RecursivelyGatherSizes(/*bReferencers=*/false, GraphRootIdentifiers, CurrentCollectionPackages, DependencyDepth, MaxSearchDependencyDepth, VisitedDependencySizeNames, DependencyNodeSizes);
+			RecursivelyGatherSizes(/*bReferencers=*/false, GraphRootIdentifiers, CurrentCollectionPackages, DependencyDepth, Settings->GetSearchDependencyDepthLimit(), VisitedDependencySizeNames, DependencyNodeSizes);
 		}
 
 		TSet<FName> AllPackageNames;
@@ -409,7 +225,7 @@ UEdGraphNode_Reference* UEdGraph_ReferenceViewer::ConstructNodes(const TArray<FA
 			}
 		};
 
-		if (bIsShowReferencers)
+		if (Settings->IsShowReferencers())
 		{
 			for (const FAssetIdentifier& AssetId : VisitedReferencerSizeNames)
 			{
@@ -417,7 +233,7 @@ UEdGraphNode_Reference* UEdGraph_ReferenceViewer::ConstructNodes(const TArray<FA
 			}
 		}
 
-		if (bIsShowDependencies)
+		if (Settings->IsShowDependencies())
 		{
 			for (const FAssetIdentifier& AssetId : VisitedDependencySizeNames)
 			{
@@ -442,34 +258,34 @@ UEdGraphNode_Reference* UEdGraph_ReferenceViewer::ConstructNodes(const TArray<FA
 		// Create the root node
 		RootNode = CreateReferenceNode();
 		bool bRootIsDuplicated = DependencyNodeInfos[GraphRootIdentifiers[0]].IsADuplicate() || ReferenceNodeInfos[GraphRootIdentifiers[0]].IsADuplicate();
-		RootNode->SetupReferenceNode(GraphRootOrigin, GraphRootIdentifiers, PackagesToAssetDataMap.FindRef(GraphRootIdentifiers[0].PackageName), /*bInAllowThumbnail = */ !bIsCompactMode, /*bIsDuplicate*/ bRootIsDuplicated);
+		RootNode->SetupReferenceNode(GraphRootOrigin, GraphRootIdentifiers, PackagesToAssetDataMap.FindRef(GraphRootIdentifiers[0].PackageName), /*bInAllowThumbnail = */ !Settings->IsCompactMode(), /*bIsDuplicate*/ bRootIsDuplicated);
 
 		if (bUseNodeInfos) // @LULU true = using new stuff, false = Not using new stuff 
 		{
-			if (bIsShowReferencers)
+			if (Settings->IsShowReferencers())
 			{
-				RecursivelyCreateNodes(true, GraphRootIdentifiers[0], GraphRootOrigin, GraphRootIdentifiers[0], RootNode, ReferenceNodeInfos, 0, MaxSearchReferencerDepth, /*bIsRoot*/ true);
+				RecursivelyCreateNodes(true, GraphRootIdentifiers[0], GraphRootOrigin, GraphRootIdentifiers[0], RootNode, ReferenceNodeInfos, 0, Settings->GetSearchReferencerDepthLimit(), /*bIsRoot*/ true);
 			}
 
-			if (bIsShowDependencies)
+			if (Settings->IsShowDependencies())
 			{
-				RecursivelyCreateNodes(false, GraphRootIdentifiers[0], GraphRootOrigin, GraphRootIdentifiers[0], RootNode, DependencyNodeInfos, 0, MaxSearchDependencyDepth, /*bIsRoot*/ true);
+				RecursivelyCreateNodes(false, GraphRootIdentifiers[0], GraphRootOrigin, GraphRootIdentifiers[0], RootNode, DependencyNodeInfos, 0, Settings->GetSearchDependencyDepthLimit(), /*bIsRoot*/ true);
 			}
 		}
 		else {
 
-			if (bIsShowReferencers)
+			if (Settings->IsShowReferencers())
 			{
 				TSet<FAssetIdentifier> VisitedReferencerNames;
 				const int32 VisitedReferencerDepth = 1;
-				RecursivelyConstructNodes(/*bReferencers=*/true, RootNode, GraphRootIdentifiers, GraphRootOrigin, ReferencerNodeSizes, PackagesToAssetDataMap, CurrentCollectionPackages, VisitedReferencerDepth, MaxSearchReferencerDepth, VisitedReferencerNames);
+				RecursivelyConstructNodes(/*bReferencers=*/true, RootNode, GraphRootIdentifiers, GraphRootOrigin, ReferencerNodeSizes, PackagesToAssetDataMap, CurrentCollectionPackages, VisitedReferencerDepth, Settings->GetSearchReferencerDepthLimit(), VisitedReferencerNames);
 			}
 
-			if (bIsShowDependencies)
+			if (Settings->IsShowDependencies())
 			{
 				TSet<FAssetIdentifier> VisitedDependencyNames;
 				const int32 VisitedDependencyDepth = 1;
-				RecursivelyConstructNodes(/*bReferencers=*/false, RootNode, GraphRootIdentifiers, GraphRootOrigin, DependencyNodeSizes, PackagesToAssetDataMap, CurrentCollectionPackages, VisitedDependencyDepth, MaxSearchDependencyDepth, VisitedDependencyNames);
+				RecursivelyConstructNodes(/*bReferencers=*/false, RootNode, GraphRootIdentifiers, GraphRootOrigin, DependencyNodeSizes, PackagesToAssetDataMap, CurrentCollectionPackages, VisitedDependencyDepth, Settings->GetSearchDependencyDepthLimit(), VisitedDependencyNames);
 			}
 		}
 	}
@@ -557,7 +373,7 @@ void UEdGraph_ReferenceViewer::GetSortedLinks(const TArray<FAssetIdentifier>& Id
 		}
 
 		// Collection Filter
-		else if ( ShouldFilterByCollection() && It.Key().IsPackage() && !CurrentCollectionPackages.Contains(It.Key().PackageName))
+		else if (ShouldFilterByCollection() && It.Key().IsPackage() && !CurrentCollectionPackages.Contains(It.Key().PackageName))
 		{
 			It.RemoveCurrent();
 		}
@@ -568,12 +384,12 @@ bool UEdGraph_ReferenceViewer::IsPackageIdentifierPassingFilter(const FAssetIden
 {
 	if (!InAssetIdentifier.IsValue())
 	{
-		if (!bIsShowNativePackages && InAssetIdentifier.PackageName.ToString().StartsWith(TEXT("/Script")))
+		if (!Settings->IsShowNativePackages() && InAssetIdentifier.PackageName.ToString().StartsWith(TEXT("/Script")))
 		{
 			return false;
 		}
 
-		if (bIsShowFilteredPackagesOnly && IsPackageNamePassingFilterCallback.IsSet() && !(*IsPackageNamePassingFilterCallback)(InAssetIdentifier.PackageName))
+		if (Settings->IsShowFilteredPackagesOnly() && IsPackageNamePassingFilterCallback.IsSet() && !(*IsPackageNamePassingFilterCallback)(InAssetIdentifier.PackageName))
 		{
 			return false;
 		}
@@ -600,7 +416,7 @@ UEdGraph_ReferenceViewer::RecursivelyPopulateNodeInfos(bool bInReferencers, cons
 		for (const TPair<FAssetIdentifier, EDependencyPinCategory>& Pair : ReferenceLinks)
 		{
 			FAssetIdentifier ChildId = Pair.Key;
-			if (!ExceedsMaxSearchBreadth(Breadth))// Breadth < MaxSearchBreadth)
+			if (!ExceedsMaxSearchBreadth(Breadth))
 			{
 				// Only Gather Children the first time
 				if (!InNodeInfos.Contains(ChildId))
@@ -614,7 +430,7 @@ UEdGraph_ReferenceViewer::RecursivelyPopulateNodeInfos(bool bInReferencers, cons
 					Breadth++;
 				}
 
-				else if (bIsShowDuplicates && !InNodeInfos[ChildId].Parents.Contains(InAssetId))
+				else if (Settings->IsShowDuplicates() && !InNodeInfos[ChildId].Parents.Contains(InAssetId))
 				{
 					InNodeInfos[ChildId].Parents.Emplace(InAssetId);
 					InNodeInfos[InAssetId].Children.Emplace(Pair);
@@ -625,7 +441,7 @@ UEdGraph_ReferenceViewer::RecursivelyPopulateNodeInfos(bool bInReferencers, cons
 			else
 			{
 				// Count the overflow nodes to report in the UI but otherwise skip adding them 
-				if (bIsShowDuplicates || !InNodeInfos.Contains(ChildId))
+				if (Settings->IsShowDuplicates() || !InNodeInfos.Contains(ChildId))
 				{
 					InNodeInfos[InAssetId].OverflowCount++;
 					Breadth++;
@@ -717,7 +533,7 @@ UEdGraphNode_Reference* UEdGraph_ReferenceViewer::RecursivelyCreateNodes(bool bI
 	else
 	{
 		NewNode = CreateReferenceNode();
-		NewNode->SetupReferenceNode(InNodeLoc, {InAssetId}, NodeInfo.AssetData, /*bInAllowThumbnail*/ !bIsCompactMode, /*bIsADuplicate*/ NodeInfo.Parents.Num() > 1);
+		NewNode->SetupReferenceNode(InNodeLoc, {InAssetId}, NodeInfo.AssetData, /*bInAllowThumbnail*/ !Settings->IsCompactMode(), /*bIsADuplicate*/ NodeInfo.Parents.Num() > 1);
 		NodeProvSize = NodeInfo.ProvisionSize(InParentId);
 	}
 
@@ -727,10 +543,10 @@ UEdGraphNode_Reference* UEdGraph_ReferenceViewer::RecursivelyCreateNodes(bool bI
 	{
 
 		// position the children nodes
-		const int32 ColumnWidth = bIsCompactMode ? 400 : 800;
+		const int32 ColumnWidth = Settings->IsCompactMode() ? 400 : 800;
 		ChildLoc.X += bInReferencers ? -ColumnWidth : ColumnWidth;
 
-		const int32 NodeSizeY = bIsCompactMode ? 100 : 200;
+		const int32 NodeSizeY = Settings->IsCompactMode() ? 100 : 200;
 		ChildLoc.Y -= (NodeProvSize - 1) * NodeSizeY * 0.5 ;
 
 		int32 Breadth = 0;
@@ -770,7 +586,7 @@ UEdGraphNode_Reference* UEdGraph_ReferenceViewer::RecursivelyCreateNodes(bool bI
 
 			if ( ensure(OverflowNode) )
 			{
-				OverflowNode->SetAllowThumbnail(!bIsCompactMode);
+				OverflowNode->SetAllowThumbnail(!Settings->IsCompactMode());
 				OverflowNode->SetReferenceNodeCollapsed(RefNodeLoc, NodeInfo.OverflowCount);
 
 				if ( bInReferencers )
@@ -803,7 +619,7 @@ UEdGraphNode_Reference* UEdGraph_ReferenceViewer::RecursivelyConstructNodes(bool
 	else
 	{
 		NewNode = CreateReferenceNode();
-		NewNode->SetupReferenceNode(NodeLoc, Identifiers, PackagesToAssetDataMap.FindRef(Identifiers[0].PackageName), /*bInAllowThumbnail = */ !bIsCompactMode, /*bIsDuplicate*/ false);
+		NewNode->SetupReferenceNode(NodeLoc, Identifiers, PackagesToAssetDataMap.FindRef(Identifiers[0].PackageName), /*bInAllowThumbnail = */ !Settings->IsCompactMode(), /*bIsDuplicate*/ false);
 	}
 
 	TMap<FAssetIdentifier, EDependencyPinCategory> Referencers;
@@ -816,15 +632,15 @@ UEdGraphNode_Reference* UEdGraph_ReferenceViewer::RecursivelyConstructNodes(bool
 		if (bReferencers)
 		{
 			// Referencers go left
-			ReferenceNodeLoc.X -= bIsCompactMode ? 400 : 800;
+			ReferenceNodeLoc.X -= Settings->IsCompactMode() ? 400 : 800;
 		}
 		else
 		{
 			// Dependencies go right
-			ReferenceNodeLoc.X += bIsCompactMode ? 400 : 800;
+			ReferenceNodeLoc.X += Settings->IsCompactMode() ? 400 : 800;
 		}
 
-		const int32 NodeSizeY = bIsCompactMode ? 100 : 200;
+		const int32 NodeSizeY = Settings->IsCompactMode() ? 100 : 200;
 		const int32 TotalReferenceSizeY = NodeSizes.FindChecked(Identifiers[0]) * NodeSizeY;
 
 		ReferenceNodeLoc.Y -= TotalReferenceSizeY * 0.5f;
@@ -894,7 +710,7 @@ UEdGraphNode_Reference* UEdGraph_ReferenceViewer::RecursivelyConstructNodes(bool
 
 			if ( ensure(ReferenceNode) )
 			{
-				ReferenceNode->SetAllowThumbnail(!bIsCompactMode);
+				ReferenceNode->SetAllowThumbnail(!Settings->IsCompactMode());
 				ReferenceNode->SetReferenceNodeCollapsed(RefNodeLoc, NumReferencesExceedingMax);
 
 				if ( bReferencers )
@@ -920,13 +736,13 @@ const TSharedPtr<FAssetThumbnailPool>& UEdGraph_ReferenceViewer::GetAssetThumbna
 bool UEdGraph_ReferenceViewer::ExceedsMaxSearchDepth(int32 Depth, int32 MaxDepth) const
 {
 	// ExceedsMaxSearchDepth requires only greater (not equal than) because, even though the Depth is 1-based indexed (similarly to Breadth), the first index (index 0) corresponds to the root object 
-	return bLimitSearchDepth && Depth > MaxDepth;
+	return Settings->IsSearchDepthLimited() && Depth > MaxDepth;
 }
 
 bool UEdGraph_ReferenceViewer::ExceedsMaxSearchBreadth(int32 Breadth) const
 {
 	// ExceedsMaxSearchBreadth requires greater or equal than because the Breadth is 1-based indexed
-	return bLimitSearchBreadth && Breadth >= MaxSearchBreadth;
+	return Settings->IsSearchBreadthLimited() && (Breadth >=  Settings->GetSearchBreadthLimit());
 }
 
 UEdGraphNode_Reference* UEdGraph_ReferenceViewer::CreateReferenceNode()
@@ -946,5 +762,5 @@ void UEdGraph_ReferenceViewer::RemoveAllNodes()
 
 bool UEdGraph_ReferenceViewer::ShouldFilterByCollection() const
 {
-	return bEnableCollectionFilter && CurrentCollectionFilter != NAME_None;
+	return Settings->GetEnableCollectionFilter() && CurrentCollectionFilter != NAME_None;
 }
diff --git a/Engine/Plugins/Editor/AssetManagerEditor/Source/AssetManagerEditor/Private/ReferenceViewer/ReferenceViewerSettings.cpp b/Engine/Plugins/Editor/AssetManagerEditor/Source/AssetManagerEditor/Private/ReferenceViewer/ReferenceViewerSettings.cpp
new file mode 100644
--- /dev/null
+++ b/Engine/Plugins/Editor/AssetManagerEditor/Source/AssetManagerEditor/Private/ReferenceViewer/ReferenceViewerSettings.cpp
@@ -0,0 +1,190 @@
+// Copyright Epic Games, Inc. All Rights Reserved.
+
+#include "ReferenceViewer/ReferenceViewerSettings.h"
+
+bool UReferenceViewerSettings::IsSearchDepthLimited() const
+{
+	return bLimitSearchDepth;
+}
+
+bool UReferenceViewerSettings::IsSearchBreadthLimited() const
+{
+	return bLimitSearchBreadth;
+}
+
+bool UReferenceViewerSettings::IsShowSoftReferences() const
+{
+	return bIsShowSoftReferences;
+}
+
+bool UReferenceViewerSettings::IsShowHardReferences() const
+{
+	return bIsShowHardReferences;
+}
+
+bool UReferenceViewerSettings::IsShowFilteredPackagesOnly() const
+{
+	return bIsShowFilteredPackagesOnly;
+}
+
+bool UReferenceViewerSettings::IsCompactMode() const
+{
+	return bIsCompactMode;
+}
+
+bool UReferenceViewerSettings::IsShowDuplicates() const
+{
+	return bIsShowDuplicates;
+}
+
+bool UReferenceViewerSettings::IsShowEditorOnlyReferences() const
+{
+	return bIsShowEditorOnlyReferences;
+}
+
+bool UReferenceViewerSettings::IsShowManagementReferences() const
+{
+	return bIsShowManagementReferences;
+}
+
+bool UReferenceViewerSettings::IsShowSearchableNames() const
+{
+	return bIsShowSearchableNames;
+}
+
+bool UReferenceViewerSettings::IsShowNativePackages() const
+{
+	return bIsShowNativePackages;
+}
+
+bool UReferenceViewerSettings::IsShowReferencers() const
+{
+	return bIsShowReferencers;
+}
+
+bool UReferenceViewerSettings::IsShowDependencies() const
+{
+	return bIsShowDependencies;
+}
+
+void UReferenceViewerSettings::SetSearchDepthLimitEnabled(bool newEnabled)
+{
+	bLimitSearchDepth = newEnabled;
+	SaveConfig();
+}
+
+void UReferenceViewerSettings::SetSearchBreadthLimitEnabled(bool newEnabled)
+{
+	bLimitSearchBreadth = newEnabled;
+	SaveConfig();
+}
+
+void UReferenceViewerSettings::SetShowSoftReferencesEnabled(bool newEnabled)
+{
+	bIsShowSoftReferences = newEnabled;
+	SaveConfig();
+}
+
+void UReferenceViewerSettings::SetShowHardReferencesEnabled(bool newEnabled)
+{
+	bIsShowHardReferences = newEnabled;
+	SaveConfig();
+}
+
+void UReferenceViewerSettings::SetShowFilteredPackagesOnlyEnabled(bool newEnabled)
+{
+	bIsShowFilteredPackagesOnly = newEnabled;
+	SaveConfig();
+}
+
+void UReferenceViewerSettings::SetCompactModeEnabled(bool newEnabled)
+{
+	bIsCompactMode = newEnabled;
+	SaveConfig();
+}
+
+void UReferenceViewerSettings::SetShowDuplicatesEnabled(bool newEnabled)
+{
+	bIsShowDuplicates = newEnabled;
+	SaveConfig();
+}
+
+void UReferenceViewerSettings::SetShowEditorOnlyReferencesEnabled(bool newEnabled)
+{
+	bIsShowEditorOnlyReferences = newEnabled;
+	SaveConfig();
+}
+
+void UReferenceViewerSettings::SetShowManagementReferencesEnabled(bool newEnabled)
+{
+	bIsShowManagementReferences = newEnabled;
+	SaveConfig();
+}
+
+void UReferenceViewerSettings::SetShowSearchableNames(bool newEnabled)
+{
+	bIsShowSearchableNames = newEnabled;
+	SaveConfig();
+}
+
+void UReferenceViewerSettings::SetShowNativePackages(bool newEnabled)
+{
+	bIsShowNativePackages = newEnabled;
+	SaveConfig();
+}
+
+void UReferenceViewerSettings::SetShowReferencers(const bool newEnabled)
+{
+	bIsShowReferencers = newEnabled;
+	SaveConfig();
+}
+
+void UReferenceViewerSettings::SetShowDependencies(const bool newEnabled)
+{
+	bIsShowDependencies = newEnabled;
+	SaveConfig();
+}
+
+int32 UReferenceViewerSettings::GetSearchDependencyDepthLimit() const
+{
+	return MaxSearchDependencyDepth;
+}
+
+void UReferenceViewerSettings::SetSearchDependencyDepthLimit(int32 NewDepthLimit)
+{
+	MaxSearchDependencyDepth = FMath::Max(NewDepthLimit, 0);
+	SaveConfig();
+}
+
+int32 UReferenceViewerSettings::GetSearchReferencerDepthLimit() const
+{
+	return MaxSearchReferencerDepth;
+}
+
+void UReferenceViewerSettings::SetSearchReferencerDepthLimit(int32 NewDepthLimit)
+{
+	MaxSearchReferencerDepth = FMath::Max(NewDepthLimit, 0);
+	SaveConfig();
+}
+
+int32 UReferenceViewerSettings::GetSearchBreadthLimit() const
+{
+	return MaxSearchBreadth;
+}
+
+void UReferenceViewerSettings::SetSearchBreadthLimit(int32 NewBreadthLimit)
+{
+	MaxSearchBreadth = FMath::Max(NewBreadthLimit, 0);
+	SaveConfig();
+}
+
+bool UReferenceViewerSettings::GetEnableCollectionFilter() const
+{
+	return bEnableCollectionFilter;
+}
+
+void UReferenceViewerSettings::SetEnableCollectionFilter(bool bEnabled)
+{
+	bEnableCollectionFilter = bEnabled;
+	SaveConfig();
+}
diff --git a/Engine/Plugins/Editor/AssetManagerEditor/Source/AssetManagerEditor/Private/ReferenceViewer/SReferenceViewer.cpp b/Engine/Plugins/Editor/AssetManagerEditor/Source/AssetManagerEditor/Private/ReferenceViewer/SReferenceViewer.cpp
--- a/Engine/Plugins/Editor/AssetManagerEditor/Source/AssetManagerEditor/Private/ReferenceViewer/SReferenceViewer.cpp
+++ b/Engine/Plugins/Editor/AssetManagerEditor/Source/AssetManagerEditor/Private/ReferenceViewer/SReferenceViewer.cpp
@@ -97,6 +97,8 @@ SReferenceViewer::~SReferenceViewer()
 
 void SReferenceViewer::Construct(const FArguments& InArgs)
 {
+	Settings = GetMutableDefault<UReferenceViewerSettings>();
+
 	// Create an action list and register commands
 	RegisterActions();
 
@@ -507,21 +509,21 @@ void SReferenceViewer::SetGraphRootIdentifiers(const TArray<FAssetIdentifier>& N
 {
 	GraphObj->SetGraphRoot(NewGraphRootIdentifiers);
 	// Set properties
-	GraphObj->SetShowReferencers(ReferenceViewerParams.bShowReferencers);
-	GraphObj->SetShowDependencies(ReferenceViewerParams.bShowDependencies);
+	Settings->SetShowReferencers(ReferenceViewerParams.bShowReferencers);
+	Settings->SetShowDependencies(ReferenceViewerParams.bShowDependencies);
 	// Set user-interactive properties
 	FixAndHideSearchDepthLimit = ReferenceViewerParams.FixAndHideSearchDepthLimit;
 	if (FixAndHideSearchDepthLimit > 0)
 	{
-		GraphObj->SetSearchDependencyDepthLimit(FixAndHideSearchDepthLimit);
-		GraphObj->SetSearchReferencerDepthLimit(FixAndHideSearchDepthLimit);
-		GraphObj->SetSearchDepthLimitEnabled(true);
+		Settings->SetSearchDependencyDepthLimit(FixAndHideSearchDepthLimit);
+		Settings->SetSearchReferencerDepthLimit(FixAndHideSearchDepthLimit);
+		Settings->SetSearchDepthLimitEnabled(true);
 	}
 	FixAndHideSearchBreadthLimit = ReferenceViewerParams.FixAndHideSearchBreadthLimit;
 	if (FixAndHideSearchBreadthLimit > 0)
 	{
-		GraphObj->SetSearchBreadthLimit(FixAndHideSearchBreadthLimit);
-		GraphObj->SetSearchBreadthLimitEnabled(true);
+		Settings->SetSearchBreadthLimit(FixAndHideSearchBreadthLimit);
+		Settings->SetSearchBreadthLimitEnabled(true);
 	}
 	bShowCollectionFilter = ReferenceViewerParams.bShowCollectionFilter;
 	bShowShowReferencesOptions = ReferenceViewerParams.bShowShowReferencesOptions;
@@ -531,14 +533,14 @@ void SReferenceViewer::SetGraphRootIdentifiers(const TArray<FAssetIdentifier>& N
 	bShowShowFilteredPackagesOnly = ReferenceViewerParams.bShowShowFilteredPackagesOnly;
 	if (ReferenceViewerParams.bShowFilteredPackagesOnly.IsSet())
 	{
-		GraphObj->SetShowFilteredPackagesOnlyEnabled(ReferenceViewerParams.bShowFilteredPackagesOnly.GetValue());
+		Settings->SetShowFilteredPackagesOnlyEnabled(ReferenceViewerParams.bShowFilteredPackagesOnly.GetValue());
 	}
 	UpdateIsPassingFilterPackageCallback();
 
 	bShowCompactMode = ReferenceViewerParams.bShowCompactMode;
 	if (ReferenceViewerParams.bCompactMode.IsSet())
 	{
-		GraphObj->SetCompactModeEnabled(ReferenceViewerParams.bCompactMode.GetValue());
+		Settings->SetCompactModeEnabled(ReferenceViewerParams.bCompactMode.GetValue());
 	}
 
 	RebuildGraph();
@@ -780,86 +782,46 @@ void SReferenceViewer::OnUpdateHistoryData(FReferenceViewerHistoryData& HistoryD
 
 void SReferenceViewer::OnSearchDepthEnabledChanged( ECheckBoxState NewState )
 {
-	if ( GraphObj )
-	{
-		GraphObj->SetSearchDepthLimitEnabled(NewState == ECheckBoxState::Checked);
-		RebuildGraph();
-	}
+	Settings->SetSearchDepthLimitEnabled(NewState == ECheckBoxState::Checked);
+	RebuildGraph();
 }
 
 ECheckBoxState SReferenceViewer::IsSearchDepthEnabledChecked() const
 {
-	if ( GraphObj )
-	{
-		return GraphObj->IsSearchDepthLimited() ? ECheckBoxState::Checked : ECheckBoxState::Unchecked;
-	}
-	else
-	{
-		return ECheckBoxState::Unchecked;
-	}
+	return Settings->IsSearchDepthLimited() ? ECheckBoxState::Checked : ECheckBoxState::Unchecked;
 }
 
 int32 SReferenceViewer::GetSearchDependencyDepthCount() const
 {
-	if ( GraphObj )
-	{
-		return GraphObj->GetSearchDependencyDepthLimit();
-	}
-	else
-	{
-		return 0;
-	}
+	return Settings->GetSearchDependencyDepthLimit();
 }
 
 int32 SReferenceViewer::GetSearchReferencerDepthCount() const
 {
-	if ( GraphObj )
-	{
-		return GraphObj->GetSearchReferencerDepthLimit();
-	}
-	else
-	{
-		return 0;
-	}
+	return Settings->GetSearchReferencerDepthLimit();
 }
 
 void SReferenceViewer::OnSearchDependencyDepthCommitted(int32 NewValue)
 {
-	if ( GraphObj )
-	{
-		GraphObj->SetSearchDependencyDepthLimit(NewValue);
-		RebuildGraph();
-	}
+	Settings->SetSearchDependencyDepthLimit(NewValue);
+	RebuildGraph();
 }
 
 void SReferenceViewer::OnSearchReferencerDepthCommitted(int32 NewValue)
 {
-	if ( GraphObj )
-	{
-		GraphObj->SetSearchReferencerDepthLimit(NewValue);
-		RebuildGraph();
-	}
+	Settings->SetSearchReferencerDepthLimit(NewValue);
+	RebuildGraph();
 }
 
 void SReferenceViewer::OnSearchBreadthEnabledChanged( ECheckBoxState NewState )
 {
-	if ( GraphObj )
-	{
-		GraphObj->SetSearchBreadthLimitEnabled(NewState == ECheckBoxState::Checked);
-		RebuildGraph();
-	}
+	Settings->SetSearchBreadthLimitEnabled(NewState == ECheckBoxState::Checked);
+	RebuildGraph();
 }
 
 ECheckBoxState SReferenceViewer::IsSearchBreadthEnabledChecked() const
 {
-	if ( GraphObj )
-	{
-		return GraphObj->IsSearchBreadthLimited() ? ECheckBoxState::Checked : ECheckBoxState::Unchecked;
-	}
-	else
-	{
-		return ECheckBoxState::Unchecked;
-	}
+	return Settings->IsSearchBreadthLimited() ? ECheckBoxState::Checked : ECheckBoxState::Unchecked;
 }
 
 TSharedRef<SWidget> SReferenceViewer::GenerateCollectionFilterItem(TSharedPtr<FName> InItem)
@@ -877,28 +839,18 @@ TSharedRef<SWidget> SReferenceViewer::GenerateCollectionFilterItem(TSharedPtr<FN
 
 void SReferenceViewer::OnEnableCollectionFilterChanged(ECheckBoxState NewState)
 {
-	if (GraphObj)
+	const bool bNewValue = NewState == ECheckBoxState::Checked;
+	const bool bCurrentValue = Settings->GetEnableCollectionFilter();
+	if (bCurrentValue != bNewValue)
 	{
-		const bool bNewValue = NewState == ECheckBoxState::Checked;
-		const bool bCurrentValue = GraphObj->GetEnableCollectionFilter();
-		if (bCurrentValue != bNewValue)
-		{
-			GraphObj->SetEnableCollectionFilter(NewState == ECheckBoxState::Checked);
-			RebuildGraph();
-		}
+		Settings->SetEnableCollectionFilter(NewState == ECheckBoxState::Checked);
+		RebuildGraph();
 	}
 }
 
 ECheckBoxState SReferenceViewer::IsEnableCollectionFilterChecked() const
 {
-	if (GraphObj)
-	{
-		return GraphObj->GetEnableCollectionFilter() ? ECheckBoxState::Checked : ECheckBoxState::Unchecked;
-	}
-	else
-	{
-		return ECheckBoxState::Unchecked;
-	}
+	return Settings->GetEnableCollectionFilter() ? ECheckBoxState::Checked : ECheckBoxState::Unchecked;
 }
 
 void SReferenceViewer::UpdateCollectionsComboList()
@@ -963,7 +915,7 @@ void SReferenceViewer::HandleCollectionFilterChanged(TSharedPtr<FName> Item, ESe
 			if (CurrentFilter == NAME_None)
 			{
 				// Automatically check the box to enable the filter if the previous filter was None
-				GraphObj->SetEnableCollectionFilter(true);
+				Settings->SetEnableCollectionFilter(true);
 			}
 
 			GraphObj->SetCurrentCollectionFilter(NewFilter);
@@ -979,46 +931,36 @@ FText SReferenceViewer::GetCollectionFilterText() const
 
 void SReferenceViewer::OnShowSoftReferencesChanged()
 {
-	if (GraphObj)
-	{
-		GraphObj->SetShowSoftReferencesEnabled( !GraphObj->IsShowSoftReferences() );
-		RebuildGraph();
-	}
+	Settings->SetShowSoftReferencesEnabled( !Settings->IsShowSoftReferences() );
+	RebuildGraph();
 }
 
 bool SReferenceViewer::IsShowSoftReferencesChecked() const
 {
-	return GraphObj ? GraphObj->IsShowSoftReferences() : false;
+	return Settings->IsShowSoftReferences();
 }
 
 void SReferenceViewer::OnShowHardReferencesChanged()
 {
-	if (GraphObj)
-	{
-		GraphObj->SetShowHardReferencesEnabled(!GraphObj->IsShowHardReferences());
-		RebuildGraph();
-	}
+	Settings->SetShowHardReferencesEnabled(!Settings->IsShowHardReferences());
+	RebuildGraph();
 }
 
-
 bool SReferenceViewer::IsShowHardReferencesChecked() const
 {
-	return GraphObj ? GraphObj->IsShowHardReferences() : false;
+	return Settings->IsShowHardReferences();
 }
 
 void SReferenceViewer::OnShowFilteredPackagesOnlyChanged()
 {
-	if (GraphObj)
-	{
-		GraphObj->SetShowFilteredPackagesOnlyEnabled(!GraphObj->IsShowFilteredPackagesOnly());
-	}
+	Settings->SetShowFilteredPackagesOnlyEnabled(!Settings->IsShowFilteredPackagesOnly());
 	UpdateIsPassingFilterPackageCallback();
 }
 
 
 bool SReferenceViewer::IsShowFilteredPackagesOnlyChecked() const
 {
-	return GraphObj ? GraphObj->IsShowFilteredPackagesOnly() : false;
+	return Settings->IsShowFilteredPackagesOnly();
 }
 
 void SReferenceViewer::UpdateIsPassingFilterPackageCallback()
@@ -1030,7 +972,7 @@ void SReferenceViewer::UpdateIsPassingFilterPackageCallback()
 		TArray<FString> SearchWords;
 		SearchString.ParseIntoArrayWS(SearchWords);
 		{
-			if (GraphObj->IsShowFilteredPackagesOnly())
+			if (Settings->IsShowFilteredPackagesOnly())
 			{
 				if (SearchWords.Num() > 0)
 				{
@@ -1046,44 +988,35 @@ void SReferenceViewer::UpdateIsPassingFilterPackageCallback()
 
 void SReferenceViewer::OnCompactModeChanged()
 {
-	if (GraphObj)
-	{
-		GraphObj->SetCompactModeEnabled(!GraphObj->IsCompactMode());
-		RebuildGraph();
-	}
+	Settings->SetCompactModeEnabled(!Settings->IsCompactMode());
+	RebuildGraph();
 }
 
 bool SReferenceViewer::IsCompactModeChecked() const
 {
-	return GraphObj ? GraphObj->IsCompactMode() : false;
+	return Settings->IsCompactMode();
 }
 
 void SReferenceViewer::OnShowDuplicatesChanged()
 {
-	if (GraphObj)
-	{
-		GraphObj->SetShowDuplicatesEnabled(!GraphObj->IsShowDuplicates());
-		RebuildGraph();
-	}
+	Settings->SetShowDuplicatesEnabled(!Settings->IsShowDuplicates());
+	RebuildGraph();
 }
 
 bool SReferenceViewer::IsShowDuplicatesChecked() const
 {
-	return GraphObj ? GraphObj->IsShowDuplicates() : false;
+	return Settings->IsShowDuplicates();
 }
 
 void SReferenceViewer::OnShowEditorOnlyReferencesChanged()
 {
-	if (GraphObj)
-	{
-		GraphObj->SetShowEditorOnlyReferencesEnabled(!GraphObj->IsShowEditorOnlyReferences());
-		RebuildGraph();
-	}
+	Settings->SetShowEditorOnlyReferencesEnabled(!Settings->IsShowEditorOnlyReferences());
+	RebuildGraph();
 }
 
 bool SReferenceViewer::IsShowEditorOnlyReferencesChecked() const
 {
-	return GraphObj ? GraphObj->IsShowEditorOnlyReferences() : false;
+	return Settings->IsShowEditorOnlyReferences();
 }
 
 
@@ -1094,61 +1027,49 @@ bool SReferenceViewer::GetManagementReferencesVisibility() const
 
 void SReferenceViewer::OnShowManagementReferencesChanged()
 {
-	if (GraphObj)
-	{
-		// This can take a few seconds if it isn't ready
-		UAssetManager::Get().UpdateManagementDatabase();
+	// This can take a few seconds if it isn't ready
+	UAssetManager::Get().UpdateManagementDatabase();
 
-		GraphObj->SetShowManagementReferencesEnabled(!GraphObj->IsShowManagementReferences());
-		RebuildGraph();
-	}
+	Settings->SetShowManagementReferencesEnabled(!Settings->IsShowManagementReferences());
+	RebuildGraph();
 }
 
 bool SReferenceViewer::IsShowManagementReferencesChecked() const
 {
-	return GraphObj ? GraphObj->IsShowManagementReferences() : false;
+	return Settings->IsShowManagementReferences();
 }
 
 void SReferenceViewer::OnShowSearchableNamesChanged()
 {
-	if (GraphObj)
-	{
-		GraphObj->SetShowSearchableNames(!GraphObj->IsShowSearchableNames());
-		RebuildGraph();
-	}
+	Settings->SetShowSearchableNames(!Settings->IsShowSearchableNames());
+	RebuildGraph();
 }
 
 bool SReferenceViewer::IsShowSearchableNamesChecked() const
 {
-	return GraphObj ? GraphObj->IsShowSearchableNames() : false;
+	return Settings->IsShowSearchableNames();
 }
 
 void SReferenceViewer::OnShowNativePackagesChanged()
 {
-	if (GraphObj)
-	{
-		GraphObj->SetShowNativePackages(!GraphObj->IsShowNativePackages());
-		RebuildGraph();
-	}
+	Settings->SetShowNativePackages(!Settings->IsShowNativePackages());
+	RebuildGraph();
 }
 
 bool SReferenceViewer::IsShowNativePackagesChecked() const
 {
-	return GraphObj ? GraphObj->IsShowNativePackages() : false;
+	return Settings->IsShowNativePackages();
 }
 
 int32 SReferenceViewer::GetSearchBreadthCount() const
 {
-	return GraphObj ? GraphObj->GetSearchBreadthLimit() : 0;
+	return Settings->GetSearchBreadthLimit();
 }
 
 void SReferenceViewer::OnSearchBreadthCommitted(int32 NewValue)
 {
-	if ( GraphObj )
-	{
-		GraphObj->SetSearchBreadthLimit(NewValue);
-		RebuildGraph();
-	}
+	Settings->SetSearchBreadthLimit(NewValue);
+	RebuildGraph();
 }
 
 void SReferenceViewer::RegisterActions()
diff --git a/Engine/Plugins/Editor/AssetManagerEditor/Source/AssetManagerEditor/Private/ReferenceViewer/SReferenceViewer.h b/Engine/Plugins/Editor/AssetManagerEditor/Source/AssetManagerEditor/Private/ReferenceViewer/SReferenceViewer.h
--- a/Engine/Plugins/Editor/AssetManagerEditor/Source/AssetManagerEditor/Private/ReferenceViewer/SReferenceViewer.h
+++ b/Engine/Plugins/Editor/AssetManagerEditor/Source/AssetManagerEditor/Private/ReferenceViewer/SReferenceViewer.h
@@ -14,6 +14,7 @@
 #include "CollectionManagerTypes.h"
 #include "AssetManagerEditorModule.h"
 #include "Containers/ArrayView.h"
+#include "ReferenceViewer/ReferenceViewerSettings.h"
 
 class UEdGraph;
 class UEdGraph_ReferenceViewer;
@@ -204,6 +205,8 @@ private:
 
 	UEdGraph_ReferenceViewer* GraphObj;
 
+	UReferenceViewerSettings* Settings;
+
 	/** The temporary copy of the path text when it is actively being edited. */
 	FText TemporaryPathBeingEdited;
 
diff --git a/Engine/Plugins/Editor/AssetManagerEditor/Source/AssetManagerEditor/Public/ReferenceViewer/EdGraph_ReferenceViewer.h b/Engine/Plugins/Editor/AssetManagerEditor/Source/AssetManagerEditor/Public/ReferenceViewer/EdGraph_ReferenceViewer.h
--- a/Engine/Plugins/Editor/AssetManagerEditor/Source/AssetManagerEditor/Public/ReferenceViewer/EdGraph_ReferenceViewer.h
+++ b/Engine/Plugins/Editor/AssetManagerEditor/Source/AssetManagerEditor/Public/ReferenceViewer/EdGraph_ReferenceViewer.h
@@ -13,6 +13,7 @@
 class FAssetThumbnailPool;
 class UEdGraphNode_Reference;
 class SReferenceViewer;
+class UReferenceViewerSettings;
 enum class EDependencyPinCategory;
 
 /*
@@ -77,52 +78,12 @@ public:
 	/** Force the graph to rebuild */
 	class UEdGraphNode_Reference* RebuildGraph();
 
-	bool IsSearchDepthLimited() const;
-	bool IsSearchBreadthLimited() const;
-	bool IsShowSoftReferences() const;
-	bool IsShowHardReferences() const;
-	bool IsShowFilteredPackagesOnly() const;
-	bool IsShowEditorOnlyReferences() const;
-	bool IsShowManagementReferences() const;
-	bool IsShowSearchableNames() const;
-	bool IsShowNativePackages() const;
-	bool IsShowReferencers() const;
-	bool IsShowDependencies() const;
-	bool IsCompactMode() const;
-	bool IsShowDuplicates() const;
-
-	void SetSearchDepthLimitEnabled(bool newEnabled);
-	void SetSearchBreadthLimitEnabled(bool newEnabled);
-	void SetShowSoftReferencesEnabled(bool newEnabled);
-	void SetShowHardReferencesEnabled(bool newEnabled);
-	void SetShowFilteredPackagesOnlyEnabled(bool newEnabled);
-	void SetShowEditorOnlyReferencesEnabled(bool newEnabled);
-	void SetShowManagementReferencesEnabled(bool newEnabled);
-	void SetShowSearchableNames(bool newEnabled);
-	void SetShowNativePackages(bool newEnabled);
-	void SetShowReferencers(const bool bShouldShowReferencers);
-	void SetShowDependencies(const bool bShouldShowDependencies);
-	void SetCompactModeEnabled(bool newEnabled);
-	void SetShowDuplicatesEnabled(bool newEnabled);
-
 	using FIsPackageNamePassingFilterCallback = TFunction<bool(FName)>;
 	void SetIsPackageNamePassingFilterCallback(const TOptional<FIsPackageNamePassingFilterCallback>& InIsPackageNamePassingFilterCallback) { IsPackageNamePassingFilterCallback = InIsPackageNamePassingFilterCallback; }
 
-	int32 GetSearchReferencerDepthLimit() const;
-	void SetSearchReferencerDepthLimit(int32 NewDepthLimit);
-
-	int32 GetSearchDependencyDepthLimit() const;
-	void SetSearchDependencyDepthLimit(int32 NewDepthLimit);
-
-	int32 GetSearchBreadthLimit() const;
-	void SetSearchBreadthLimit(int32 NewBreadthLimit);
-
 	FName GetCurrentCollectionFilter() const;
 	void SetCurrentCollectionFilter(FName NewFilter);
 
-	bool GetEnableCollectionFilter() const;
-	void SetEnableCollectionFilter(bool bEnabled);
-
 	/* Temporary variable that allows reverting to deprecated layout methods */
 	bool GetUseNodeInfos() const { return bUseNodeInfos; }
 	void SetUseNodeInfos(bool InbUseNodeInfos) { bUseNodeInfos= InbUseNodeInfos; }
@@ -178,39 +139,22 @@ private:
 	TArray<FAssetIdentifier> CurrentGraphRootIdentifiers;
 	FIntPoint CurrentGraphRootOrigin;
 
-	int32 MaxSearchReferencerDepth; // How deep to search references
-	int32 MaxSearchDependencyDepth; // How deep to search dependanies
 	int32 MaxSearchBreadth;
 
 	/** Current collection filter. NAME_None for no filter */
 	FName CurrentCollectionFilter;
-	bool bEnableCollectionFilter;
-
-	bool bLimitSearchDepth;
-	bool bLimitSearchBreadth;
-	bool bIsShowSoftReferences;
-	bool bIsShowHardReferences;
-	bool bIsShowEditorOnlyReferences;
-	bool bIsShowManagementReferences;
-	bool bIsShowSearchableNames;
-	bool bIsShowNativePackages;
-	/* Whether to display the Referencers */
-	bool bIsShowReferencers;
-	/* Whether to display the Dependencies */
-	bool bIsShowDependencies;
-	/* Whether to show duplicate asset references */
-	bool bIsShowDuplicates;
+
 
 	/* This is a convenience toggle to switch between the old & new methods for computing & displaying the graph */
 	bool bUseNodeInfos;
 
-	bool bIsShowFilteredPackagesOnly;
 	TOptional<FIsPackageNamePassingFilterCallback> IsPackageNamePassingFilterCallback;
 
 	/** List of packages the current collection filter allows */
 	TSet<FName> CurrentCollectionPackages;
 
-	bool bIsCompactMode;
+
+	UReferenceViewerSettings* Settings;
 
 	friend SReferenceViewer;
 };
diff --git a/Engine/Plugins/Editor/AssetManagerEditor/Source/AssetManagerEditor/Public/ReferenceViewer/ReferenceViewerSettings.h b/Engine/Plugins/Editor/AssetManagerEditor/Source/AssetManagerEditor/Public/ReferenceViewer/ReferenceViewerSettings.h
new file mode 100644
--- /dev/null
+++ b/Engine/Plugins/Editor/AssetManagerEditor/Source/AssetManagerEditor/Public/ReferenceViewer/ReferenceViewerSettings.h
@@ -0,0 +1,141 @@
+// Copyright Epic Games, Inc. All Rights Reserved.
+#pragma once
+
+#include "CoreMinimal.h"
+#include "UObject/ObjectMacros.h"
+#include "UObject/Object.h"
+
+#include "ReferenceViewerSettings.generated.h"
+
+/**
+ *  Project based Reference Viewer Saved Settings 
+ */
+
+UCLASS(config=EditorPerProjectUserSettings)
+class UReferenceViewerSettings : public UObject
+{
+public:
+	GENERATED_BODY()
+
+	bool IsSearchDepthLimited() const;
+	void SetSearchDepthLimitEnabled(bool newEnabled);
+
+	bool IsShowReferencers() const;
+	void SetShowReferencers(const bool bShouldShowReferencers);
+
+	int32 GetSearchReferencerDepthLimit() const;
+	void SetSearchReferencerDepthLimit(int32 NewDepthLimit);
+
+	bool IsShowDependencies() const;
+	void SetShowDependencies(const bool bShouldShowDependencies);
+
+	int32 GetSearchDependencyDepthLimit() const;
+	void SetSearchDependencyDepthLimit(int32 NewDepthLimit);
+
+	bool IsSearchBreadthLimited() const;
+	void SetSearchBreadthLimitEnabled(bool newEnabled);
+
+	int32 GetSearchBreadthLimit() const;
+	void SetSearchBreadthLimit(int32 NewBreadthLimit);
+
+	bool GetEnableCollectionFilter() const;
+	void SetEnableCollectionFilter(bool bEnabled);
+
+	bool IsShowSoftReferences() const;
+	void SetShowSoftReferencesEnabled(bool newEnabled);
+
+	bool IsShowHardReferences() const;
+	void SetShowHardReferencesEnabled(bool newEnabled);
+
+	bool IsShowEditorOnlyReferences() const;
+	void SetShowEditorOnlyReferencesEnabled(bool newEnabled);
+
+	bool IsShowManagementReferences() const;
+	void SetShowManagementReferencesEnabled(bool newEnabled);
+
+	bool IsShowSearchableNames() const;
+	void SetShowSearchableNames(bool newEnabled);
+
+	bool IsShowNativePackages() const;
+	void SetShowNativePackages(bool newEnabled);
+
+	bool IsShowDuplicates() const;
+	void SetShowDuplicatesEnabled(bool newEnabled);
+
+	bool IsShowFilteredPackagesOnly() const;
+	void SetShowFilteredPackagesOnlyEnabled(bool newEnabled);
+
+	bool IsCompactMode() const;
+	void SetCompactModeEnabled(bool newEnabled);
+
+
+private:
+	/* Whether to limit the search depth for Referencers & Dependencies */
+	UPROPERTY(config)
+	bool bLimitSearchDepth;
+	
+	/* Whether to display the Referencers */
+	UPROPERTY(config)
+	bool bIsShowReferencers;
+	
+	/* How deep to search references */
+	UPROPERTY(config)
+	int32 MaxSearchReferencerDepth; 
+	
+	/* Whether to display the Dependencies */
+	UPROPERTY(config)
+	bool bIsShowDependencies;
+	
+	/* How deep to search dependanies */
+	UPROPERTY(config)
+	int32 MaxSearchDependencyDepth; 
+	
+	/* Whether or not to limit how many siblings can appear */
+	UPROPERTY(config)
+	bool bLimitSearchBreadth;
+	
+	/* The max number of siblings that can appear from a node */
+	UPROPERTY(config)
+	int32 MaxSearchBreadth;
+	
+	/* Whether or not to filter from a collection */
+	UPROPERTY(config)
+	bool bEnableCollectionFilter;
+	
+	/* Show/Hide Soft References */
+	UPROPERTY(config)
+	bool bIsShowSoftReferences;
+	
+	/* Show/Hide Hard References */
+	UPROPERTY(config)
+	bool bIsShowHardReferences;
+	
+	/* Show/Hide EditorOnly References */
+	UPROPERTY(config)
+	bool bIsShowEditorOnlyReferences;
+	
+	/* Show/Hide Management Assets (i.e. PrimaryAssetIds) */
+	UPROPERTY(config)
+	bool bIsShowManagementReferences;
+	
+	/* Show/Hide Searchable Names (i.e. Gameplay Tags) */
+	UPROPERTY(config)
+	bool bIsShowSearchableNames;
+	
+	/* Show/Hide Native Packages  */
+	UPROPERTY(config)
+	bool bIsShowNativePackages;
+	
+	/* Whether to show duplicate asset references */
+	UPROPERTY(config)
+	bool bIsShowDuplicates;
+	
+	/* Whether to filter the search results or just select them  */
+	UPROPERTY(config)
+	bool bIsShowFilteredPackagesOnly;
+	
+	/* Whether to show the nodes in a compact (no thumbnail) view */
+	UPROPERTY(config)
+	bool bIsCompactMode;
+
+};
