Reference Viewer Asset Paths

- Adds option to show/hide asset path as a comment (Shortcut: P)
- Adds shortcut CTRL+ C to copy selected package paths

diff --git a/Engine/Plugins/Editor/AssetManagerEditor/Source/AssetManagerEditor/Private/AssetManagerEditorCommands.cpp b/Engine/Plugins/Editor/AssetManagerEditor/Source/AssetManagerEditor/Private/AssetManagerEditorCommands.cpp
--- a/Engine/Plugins/Editor/AssetManagerEditor/Source/AssetManagerEditor/Private/AssetManagerEditorCommands.cpp
+++ b/Engine/Plugins/Editor/AssetManagerEditor/Source/AssetManagerEditor/Private/AssetManagerEditorCommands.cpp
@@ -50,6 +50,10 @@ void FAssetManagerEditorCommands::RegisterCommands()
 	UI_COMMAND(CompactMode, "Compact Mode", "Toggles Compact Mode", EUserInterfaceActionType::ToggleButton, FInputChord(EKeys::C));
 	UI_COMMAND(FilterSearch, "Filter Search Results", "Toggles filtering of search results", EUserInterfaceActionType::ToggleButton, FInputChord(EKeys::F));
 
+	UI_COMMAND(ShowCommentPath, "Show Asset Path", "Toggles visibility of the Asset Path shown as a comment", EUserInterfaceActionType::ToggleButton, FInputChord(EKeys::P));
+
+	UI_COMMAND(CopyPaths, "Copy the Asset Path", "Copies the Asset Path to the Clipboard", EUserInterfaceActionType::Button, FInputChord(EModifierKey::Control, EKeys::C));
+
 	UI_COMMAND(ReCenterGraph, "Re-Center Graph", "Re-centers the graph on this node, showing all referencers and references for this asset instead", EUserInterfaceActionType::Button, FInputChord());
 	UI_COMMAND(CopyReferencedObjects, "Copy Referenced Objects List", "Copies the list of objects that the selected asset references to the clipboard.", EUserInterfaceActionType::Button, FInputChord());
 	UI_COMMAND(CopyReferencingObjects, "Copy Referencing Objects List", "Copies the list of objects that reference the selected asset to the clipboard.", EUserInterfaceActionType::Button, FInputChord());
diff --git a/Engine/Plugins/Editor/AssetManagerEditor/Source/AssetManagerEditor/Private/AssetManagerEditorCommands.h b/Engine/Plugins/Editor/AssetManagerEditor/Source/AssetManagerEditor/Private/AssetManagerEditorCommands.h
--- a/Engine/Plugins/Editor/AssetManagerEditor/Source/AssetManagerEditor/Private/AssetManagerEditorCommands.h
+++ b/Engine/Plugins/Editor/AssetManagerEditor/Source/AssetManagerEditor/Private/AssetManagerEditorCommands.h
@@ -85,6 +85,12 @@ public:
 	// Toggles Compact Mode
 	TSharedPtr<FUICommandInfo> CompactMode;
 
+	// Toggles Path Comment 
+	TSharedPtr<FUICommandInfo> ShowCommentPath;
+
+	// Copies the selected Asset Paths to the Clipboard
+	TSharedPtr<FUICommandInfo> CopyPaths;
+
 	// Toggles whether search results are filtered or just selected 
 	TSharedPtr<FUICommandInfo> FilterSearch;
 
diff --git a/Engine/Plugins/Editor/AssetManagerEditor/Source/AssetManagerEditor/Private/ReferenceViewer/EdGraphNode_Reference.cpp b/Engine/Plugins/Editor/AssetManagerEditor/Source/AssetManagerEditor/Private/ReferenceViewer/EdGraphNode_Reference.cpp
--- a/Engine/Plugins/Editor/AssetManagerEditor/Source/AssetManagerEditor/Private/ReferenceViewer/EdGraphNode_Reference.cpp
+++ b/Engine/Plugins/Editor/AssetManagerEditor/Source/AssetManagerEditor/Private/ReferenceViewer/EdGraphNode_Reference.cpp
@@ -78,6 +78,11 @@ void UEdGraphNode_Reference::SetupReferenceNode(const FIntPoint& NodeLoc, const
 
 		// append the type so it shows up on the extra line
 		NodeTitle = FText::FromString(FString::Printf(TEXT("%s\n%s"), *MainAssetName, *AssetTypeName));
+
+		if (bIsPackage)
+		{
+			NodeComment = First.PackageName.ToString();
+		}
 	}
 	else
 	{
diff --git a/Engine/Plugins/Editor/AssetManagerEditor/Source/AssetManagerEditor/Private/ReferenceViewer/EdGraph_ReferenceViewer.cpp b/Engine/Plugins/Editor/AssetManagerEditor/Source/AssetManagerEditor/Private/ReferenceViewer/EdGraph_ReferenceViewer.cpp
--- a/Engine/Plugins/Editor/AssetManagerEditor/Source/AssetManagerEditor/Private/ReferenceViewer/EdGraph_ReferenceViewer.cpp
+++ b/Engine/Plugins/Editor/AssetManagerEditor/Source/AssetManagerEditor/Private/ReferenceViewer/EdGraph_ReferenceViewer.cpp
@@ -259,6 +259,7 @@ UEdGraphNode_Reference* UEdGraph_ReferenceViewer::ConstructNodes(const TArray<FA
 		RootNode = CreateReferenceNode();
 		bool bRootIsDuplicated = DependencyNodeInfos[GraphRootIdentifiers[0]].IsADuplicate() || ReferenceNodeInfos[GraphRootIdentifiers[0]].IsADuplicate();
 		RootNode->SetupReferenceNode(GraphRootOrigin, GraphRootIdentifiers, PackagesToAssetDataMap.FindRef(GraphRootIdentifiers[0].PackageName), /*bInAllowThumbnail = */ !Settings->IsCompactMode(), /*bIsDuplicate*/ bRootIsDuplicated);
+		RootNode->SetMakeCommentBubbleVisible(Settings->IsShowPath());
 
 		if (bUseNodeInfos) // @LULU true = using new stuff, false = Not using new stuff 
 		{
@@ -534,6 +535,7 @@ UEdGraphNode_Reference* UEdGraph_ReferenceViewer::RecursivelyCreateNodes(bool bI
 	{
 		NewNode = CreateReferenceNode();
 		NewNode->SetupReferenceNode(InNodeLoc, {InAssetId}, NodeInfo.AssetData, /*bInAllowThumbnail*/ !Settings->IsCompactMode(), /*bIsADuplicate*/ NodeInfo.Parents.Num() > 1);
+		NewNode->SetMakeCommentBubbleVisible(Settings->IsShowPath());
 		NodeProvSize = NodeInfo.ProvisionSize(InParentId);
 	}
 
@@ -543,10 +545,12 @@ UEdGraphNode_Reference* UEdGraph_ReferenceViewer::RecursivelyCreateNodes(bool bI
 	{
 
 		// position the children nodes
-		const int32 ColumnWidth = Settings->IsCompactMode() ? 400 : 800;
+		const int32 ColumnWidth = Settings->IsCompactMode() ? 500 : 800;
 		ChildLoc.X += bInReferencers ? -ColumnWidth : ColumnWidth;
 
-		const int32 NodeSizeY = Settings->IsCompactMode() ? 100 : 200;
+		int32 NodeSizeY = Settings->IsCompactMode() ? 100 : 200;
+		NodeSizeY += Settings->IsShowPath() ? 40 : 0;
+
 		ChildLoc.Y -= (NodeProvSize - 1) * NodeSizeY * 0.5 ;
 
 		int32 Breadth = 0;
diff --git a/Engine/Plugins/Editor/AssetManagerEditor/Source/AssetManagerEditor/Private/ReferenceViewer/ReferenceViewerSettings.cpp b/Engine/Plugins/Editor/AssetManagerEditor/Source/AssetManagerEditor/Private/ReferenceViewer/ReferenceViewerSettings.cpp
--- a/Engine/Plugins/Editor/AssetManagerEditor/Source/AssetManagerEditor/Private/ReferenceViewer/ReferenceViewerSettings.cpp
+++ b/Engine/Plugins/Editor/AssetManagerEditor/Source/AssetManagerEditor/Private/ReferenceViewer/ReferenceViewerSettings.cpp
@@ -188,3 +188,15 @@ void UReferenceViewerSettings::SetEnableCollectionFilter(bool bEnabled)
 	bEnableCollectionFilter = bEnabled;
 	SaveConfig();
 }
+
+bool UReferenceViewerSettings::IsShowPath() const
+{
+	return bIsShowPath;
+}
+
+void UReferenceViewerSettings::SetShowPathEnabled(bool bEnabled)
+{
+	bIsShowPath = bEnabled;
+	SaveConfig();
+}
+
diff --git a/Engine/Plugins/Editor/AssetManagerEditor/Source/AssetManagerEditor/Private/ReferenceViewer/SReferenceNode.cpp b/Engine/Plugins/Editor/AssetManagerEditor/Source/AssetManagerEditor/Private/ReferenceViewer/SReferenceNode.cpp
--- a/Engine/Plugins/Editor/AssetManagerEditor/Source/AssetManagerEditor/Private/ReferenceViewer/SReferenceNode.cpp
+++ b/Engine/Plugins/Editor/AssetManagerEditor/Source/AssetManagerEditor/Private/ReferenceViewer/SReferenceNode.cpp
@@ -244,7 +244,7 @@ void SReferenceNode::UpdateGraphNode()
 	];
 	// Create comment bubble if comment text is valid
 	GetNodeObj()->bCommentBubbleVisible = !GetNodeObj()->NodeComment.IsEmpty();
-	if( GetNodeObj()->bCommentBubbleVisible )
+	if( GetNodeObj()->ShouldMakeCommentBubbleVisible() && GetNodeObj()->bCommentBubbleVisible)
 	{
 		TSharedPtr<SCommentBubble> CommentBubble;
 
diff --git a/Engine/Plugins/Editor/AssetManagerEditor/Source/AssetManagerEditor/Private/ReferenceViewer/SReferenceViewer.cpp b/Engine/Plugins/Editor/AssetManagerEditor/Source/AssetManagerEditor/Private/ReferenceViewer/SReferenceViewer.cpp
--- a/Engine/Plugins/Editor/AssetManagerEditor/Source/AssetManagerEditor/Private/ReferenceViewer/SReferenceViewer.cpp
+++ b/Engine/Plugins/Editor/AssetManagerEditor/Source/AssetManagerEditor/Private/ReferenceViewer/SReferenceViewer.cpp
@@ -1274,6 +1274,41 @@ void SReferenceViewer::RegisterActions()
 		FAssetManagerEditorCommands::Get().ViewAssetAudit,
 		FExecuteAction::CreateSP(this, &SReferenceViewer::ViewAssetAudit),
 		FCanExecuteAction::CreateSP(this, &SReferenceViewer::HasAtLeastOneRealNodeSelected));
+
+	ReferenceViewerActions->MapAction(
+		FAssetManagerEditorCommands::Get().ShowCommentPath,
+		FExecuteAction::CreateLambda([this] { 
+			Settings->SetShowPathEnabled(!Settings->IsShowPath());
+			RebuildGraph();
+		}),
+		FCanExecuteAction(),	
+		FIsActionChecked::CreateLambda([this] {return Settings->IsShowPath();}));
+
+
+	ReferenceViewerActions->MapAction(
+		FAssetManagerEditorCommands::Get().CopyPaths,
+		FExecuteAction::CreateLambda([this] {
+				FString Result;
+				// Build up a list of selected assets from the graph selection set
+				TSet<UObject*> SelectedNodes = GraphEditorPtr->GetSelectedNodes();
+				for (FGraphPanelSelectionSet::TConstIterator It(SelectedNodes); It; ++It)
+				{
+					if (UEdGraphNode_Reference* ReferenceNode = Cast<UEdGraphNode_Reference>(*It))
+					{
+						if (ReferenceNode->GetAssetData().IsValid())
+						{
+							Result += ReferenceNode->GetAssetData().PackageName.ToString();
+							Result += TEXT("\n");
+						}
+					}
+				}
+
+				if (Result.Len())
+				{
+					FPlatformApplicationMisc::ClipboardCopy(*Result);
+				}
+		}),
+		FCanExecuteAction::CreateSP(this, &SReferenceViewer::HasAtLeastOneRealNodeSelected));
 }
 
 void SReferenceViewer::ShowSelectionInContentBrowser()
@@ -1824,6 +1859,7 @@ TSharedRef<SWidget> SReferenceViewer::GetShowMenuContent()
 	MenuBuilder.AddMenuEntry(FAssetManagerEditorCommands::Get().ShowDuplicates);
 	MenuBuilder.AddMenuEntry(FAssetManagerEditorCommands::Get().FilterSearch);
 	MenuBuilder.AddMenuEntry(FAssetManagerEditorCommands::Get().CompactMode);
+	MenuBuilder.AddMenuEntry(FAssetManagerEditorCommands::Get().ShowCommentPath);
 	MenuBuilder.EndSection();
 
 	return MenuBuilder.MakeWidget();
diff --git a/Engine/Plugins/Editor/AssetManagerEditor/Source/AssetManagerEditor/Public/ReferenceViewer/ReferenceViewerSettings.h b/Engine/Plugins/Editor/AssetManagerEditor/Source/AssetManagerEditor/Public/ReferenceViewer/ReferenceViewerSettings.h
--- a/Engine/Plugins/Editor/AssetManagerEditor/Source/AssetManagerEditor/Public/ReferenceViewer/ReferenceViewerSettings.h
+++ b/Engine/Plugins/Editor/AssetManagerEditor/Source/AssetManagerEditor/Public/ReferenceViewer/ReferenceViewerSettings.h
@@ -68,6 +68,8 @@ public:
 	bool IsCompactMode() const;
 	void SetCompactModeEnabled(bool newEnabled);
 
+	bool IsShowPath() const;
+	void SetShowPathEnabled(bool newEnabled);
 
 private:
 	/* Whether to limit the search depth for Referencers & Dependencies */
@@ -138,4 +140,8 @@ private:
 	UPROPERTY(config)
 	bool bIsCompactMode;
 
+	/* Whether to show the package's path as a comment */
+	UPROPERTY(config)
+	bool bIsShowPath;
+
 };
