Reference Viewer Enhancements

- Adds Asset Type Filters
- Auto Filters Mode
- Updates hotkeys so "F" now frames and "A" filters on/off
- Fixes Issue where scripts were showing up as "None"
- Migrated UI to use a Toolbar which populates from the commands
- Removes CVar and depricated Layout methods
- Added optimizations when changing just the breadth, filters and showing/hiding duplicates or comments

diff --git a/Engine/Config/BaseEditorPerProjectUserSettings.ini b/Engine/Config/BaseEditorPerProjectUserSettings.ini
--- a/Engine/Config/BaseEditorPerProjectUserSettings.ini
+++ b/Engine/Config/BaseEditorPerProjectUserSettings.ini
@@ -963,6 +963,10 @@ bIsShowNativePackages=False
 bIsShowDuplicates=True
 bIsShowFilteredPackagesOnly=False
 bIsCompactMode=False
+bIsShowPath=False
+bFiltersEnabled=True
+bAutoUpdateFilters=False
+
 
 [/Script/BlueprintHeaderView.BlueprintHeaderViewSettings]
 SyntaxColors=(Comment=(R=0.258183,G=0.539479,B=0.068478,A=1.000000),Error=(R=0.863157,G=0.035601,B=0.035601,A=1.000000),Macro=(R=0.356400,G=0.040915,B=0.520996,A=1.000000),Typename=(R=0.000000,G=0.300000,B=0.300000,A=1.000000),Identifier=(R=1.000000,G=1.000000,B=1.000000,A=1.000000),Keyword=(R=0.019382,G=0.496933,B=1.000000,A=1.000000))
diff --git a/Engine/Plugins/Editor/AssetManagerEditor/Source/AssetManagerEditor/Private/AssetManagerEditorCommands.cpp b/Engine/Plugins/Editor/AssetManagerEditor/Source/AssetManagerEditor/Private/AssetManagerEditorCommands.cpp
--- a/Engine/Plugins/Editor/AssetManagerEditor/Source/AssetManagerEditor/Private/AssetManagerEditorCommands.cpp
+++ b/Engine/Plugins/Editor/AssetManagerEditor/Source/AssetManagerEditor/Private/AssetManagerEditorCommands.cpp
@@ -24,7 +24,7 @@ void FAssetManagerEditorCommands::RegisterCommands()
 	UI_COMMAND(ViewAssetAudit, "Audit Assets...", "Opens the Asset Audit UI and displays information about the selected assets", EUserInterfaceActionType::Button, FInputChord(EModifierKey::Shift | EModifierKey::Alt, EKeys::A));
 
 	UI_COMMAND(OpenSelectedInAssetEditor, "Edit...", "Opens the selected asset in the relevant editor.", EUserInterfaceActionType::Button, FInputChord(EModifierKey::Control, EKeys::E));
-	UI_COMMAND(ZoomToFit, "Zoom to Fit", "Zoom in and center the view on the selected item", EUserInterfaceActionType::Button, FInputChord(EKeys::Home));
+	UI_COMMAND(ZoomToFit, "Zoom to Fit", "Zoom in and center the view on the selected item", EUserInterfaceActionType::Button, FInputChord(EKeys::F));
 
 	UI_COMMAND(IncreaseReferencerSearchDepth, "Increase Referencer Search Depth", "Increase the Referencer Search Depth", EUserInterfaceActionType::Button, FInputChord(EKeys::R));
 	UI_COMMAND(DecreaseReferencerSearchDepth, "Decrease Referencer Search Depth", "Decrease the Referencer Search Depth", EUserInterfaceActionType::Button, FInputChord(EModifierKey::Shift, EKeys::R));
@@ -43,14 +43,16 @@ void FAssetManagerEditorCommands::RegisterCommands()
 	UI_COMMAND(ShowEditorOnlyReferences, "Show EditorOnly References","Toggles visibility of EditorOnly References", EUserInterfaceActionType::ToggleButton, FInputChord(EKeys::E));
 
 	UI_COMMAND(ShowManagementReferences, "Show Management References","Toggles visibility of Management References (i.e. PrimaryAssetIDs)", EUserInterfaceActionType::ToggleButton, FInputChord(EKeys::M));
-	UI_COMMAND(ShowGamePlayTags, "Show GamePlayTags","Toggles visibility of GameplayTags", EUserInterfaceActionType::ToggleButton, FInputChord(EKeys::G));
+	UI_COMMAND(ShowGamePlayTags, "Show GamePlay Tags","Toggles visibility of GamePlay Tags", EUserInterfaceActionType::ToggleButton, FInputChord(EKeys::G));
 	UI_COMMAND(ShowNativePackages, "Show Native Packages","Toggles visibility of Native Packages", EUserInterfaceActionType::ToggleButton, FInputChord(EKeys::N));
 
 	UI_COMMAND(ShowDuplicates, "Show Duplicate References", "Toggles visibility of Duplicate References", EUserInterfaceActionType::ToggleButton, FInputChord(EKeys::U));
 	UI_COMMAND(CompactMode, "Compact Mode", "Toggles Compact Mode", EUserInterfaceActionType::ToggleButton, FInputChord(EKeys::C));
-	UI_COMMAND(FilterSearch, "Filter Search Results", "Toggles filtering of search results", EUserInterfaceActionType::ToggleButton, FInputChord(EKeys::F));
+	UI_COMMAND(FilterSearch, "Filter Search Results", "Toggles filtering of search results", EUserInterfaceActionType::ToggleButton, FInputChord());
 
 	UI_COMMAND(ShowCommentPath, "Show Asset Path", "Toggles visibility of the Asset Path shown as a comment", EUserInterfaceActionType::ToggleButton, FInputChord(EKeys::P));
+	UI_COMMAND(Filters, "Toggle the AssetType Filters", "Toggles the AssetType Filters", EUserInterfaceActionType::ToggleButton, FInputChord(EKeys::A));
+	UI_COMMAND(AutoFilters, "Auto Populate Filters", "Toggles autopopulating of AssetType Filters", EUserInterfaceActionType::ToggleButton, FInputChord());
 
 	UI_COMMAND(CopyPaths, "Copy the Asset Path", "Copies the Asset Path to the Clipboard", EUserInterfaceActionType::Button, FInputChord(EModifierKey::Control, EKeys::C));
 
diff --git a/Engine/Plugins/Editor/AssetManagerEditor/Source/AssetManagerEditor/Private/AssetManagerEditorCommands.h b/Engine/Plugins/Editor/AssetManagerEditor/Source/AssetManagerEditor/Private/AssetManagerEditorCommands.h
--- a/Engine/Plugins/Editor/AssetManagerEditor/Source/AssetManagerEditor/Private/AssetManagerEditorCommands.h
+++ b/Engine/Plugins/Editor/AssetManagerEditor/Source/AssetManagerEditor/Private/AssetManagerEditorCommands.h
@@ -88,6 +88,12 @@ public:
 	// Toggles Path Comment 
 	TSharedPtr<FUICommandInfo> ShowCommentPath;
 
+	// Toggles FilterBar Filters On/Off
+	TSharedPtr<FUICommandInfo> Filters;
+
+	// Toggles if Asset Type Filters are Auto Populated & Updated
+	TSharedPtr<FUICommandInfo> AutoFilters;
+
 	// Copies the selected Asset Paths to the Clipboard
 	TSharedPtr<FUICommandInfo> CopyPaths;
 
diff --git a/Engine/Plugins/Editor/AssetManagerEditor/Source/AssetManagerEditor/Private/ReferenceViewer/EdGraph_ReferenceViewer.cpp b/Engine/Plugins/Editor/AssetManagerEditor/Source/AssetManagerEditor/Private/ReferenceViewer/EdGraph_ReferenceViewer.cpp
--- a/Engine/Plugins/Editor/AssetManagerEditor/Source/AssetManagerEditor/Private/ReferenceViewer/EdGraph_ReferenceViewer.cpp
+++ b/Engine/Plugins/Editor/AssetManagerEditor/Source/AssetManagerEditor/Private/ReferenceViewer/EdGraph_ReferenceViewer.cpp
@@ -16,24 +16,25 @@
 #include "Engine/AssetManager.h"
 #include "Settings/EditorProjectSettings.h"
 
-ReferenceNodeInfo::ReferenceNodeInfo(const FAssetIdentifier& InAssetId, bool InbReferencers)
+FReferenceNodeInfo::FReferenceNodeInfo(const FAssetIdentifier& InAssetId, bool InbReferencers)
 	: AssetId(InAssetId)
 	, bReferencers(InbReferencers)
 	, OverflowCount(0)
 	, ChildProvisionSize(0)
+	, PassedFilters(true)
 {}
 
-bool ReferenceNodeInfo::IsFirstParent(const FAssetIdentifier& InParentId) const
+bool FReferenceNodeInfo::IsFirstParent(const FAssetIdentifier& InParentId) const
 {
 	return Parents.IsEmpty() || Parents[0] == InParentId;
 }
 
-bool ReferenceNodeInfo::IsADuplicate() const
+bool FReferenceNodeInfo::IsADuplicate() const
 {
 	return Parents.Num() > 1;
 }
 
-int32 ReferenceNodeInfo::ProvisionSize(const FAssetIdentifier& InParentId) const
+int32 FReferenceNodeInfo::ProvisionSize(const FAssetIdentifier& InParentId) const
 {
 	return IsFirstParent(InParentId) ? ChildProvisionSize : 1;
 }
@@ -44,8 +45,6 @@ UEdGraph_ReferenceViewer::UEdGraph_ReferenceViewer(const FObjectInitializer& Obj
 	AssetThumbnailPool = MakeShareable( new FAssetThumbnailPool(1024) );
 
 	Settings = GetMutableDefault<UReferenceViewerSettings>();
-
-	bUseNodeInfos = true;
 }
 
 void UEdGraph_ReferenceViewer::BeginDestroy()
@@ -135,6 +134,11 @@ void UEdGraph_ReferenceViewer::SetCurrentCollectionFilter(FName NewFilter)
 	CurrentCollectionFilter = NewFilter;
 }
 
+void UEdGraph_ReferenceViewer::SetCurrentFilterCollection(TSharedPtr< TFilterCollection< FReferenceNodeInfo& > > InFilterCollection )
+{
+	FilterCollection = InFilterCollection;
+}
+
 FAssetManagerDependencyQuery UEdGraph_ReferenceViewer::GetReferenceSearchFlags(bool bHardOnly) const
 {
 	using namespace UE::AssetRegistry;
@@ -166,7 +170,6 @@ FAssetManagerDependencyQuery UEdGraph_ReferenceViewer::GetReferenceSearchFlags(b
 UEdGraphNode_Reference* UEdGraph_ReferenceViewer::ConstructNodes(const TArray<FAssetIdentifier>& GraphRootIdentifiers, const FIntPoint& GraphRootOrigin )
 {
 	UEdGraphNode_Reference* RootNode = NULL;
-
 	if (GraphRootIdentifiers.Num() > 0)
 	{
 		// It both were false, nothing (other than the GraphRootIdentifiers) would be displayed
@@ -189,31 +192,15 @@ UEdGraphNode_Reference* UEdGraph_ReferenceViewer::ConstructNodes(const TArray<FA
 		// Create & Populate the NodeInfo Maps 
 		// Note to add an empty parent to the root so that if the root node again gets found again as a duplicate, that next parent won't be 
 		// identified as the primary root and also it will appear as having multiple parents.
-		TMap<FAssetIdentifier, ReferenceNodeInfo> ReferenceNodeInfos;
-		ReferenceNodeInfo& RootNodeInfo = ReferenceNodeInfos.FindOrAdd( GraphRootIdentifiers[0], ReferenceNodeInfo(GraphRootIdentifiers[0], true));
+		TMap<FAssetIdentifier, FReferenceNodeInfo> NewReferenceNodeInfos;
+		FReferenceNodeInfo& RootNodeInfo = NewReferenceNodeInfos.FindOrAdd( GraphRootIdentifiers[0], FReferenceNodeInfo(GraphRootIdentifiers[0], true));
 		RootNodeInfo.Parents.Emplace(FAssetIdentifier(NAME_None));
-		RecursivelyPopulateNodeInfos(true, GraphRootIdentifiers[0], ReferenceNodeInfos, 0, Settings->GetSearchReferencerDepthLimit());
+		RecursivelyPopulateNodeInfos(true, GraphRootIdentifiers[0], NewReferenceNodeInfos, 0, Settings->GetSearchReferencerDepthLimit());
 
-		TMap<FAssetIdentifier, ReferenceNodeInfo> DependencyNodeInfos;
-		ReferenceNodeInfo& DRootNodeInfo = DependencyNodeInfos.FindOrAdd( GraphRootIdentifiers[0], ReferenceNodeInfo(GraphRootIdentifiers[0], false));
+		TMap<FAssetIdentifier, FReferenceNodeInfo> NewDependencyNodeInfos;
+		FReferenceNodeInfo& DRootNodeInfo = NewDependencyNodeInfos.FindOrAdd( GraphRootIdentifiers[0], FReferenceNodeInfo(GraphRootIdentifiers[0], false));
 		DRootNodeInfo.Parents.Emplace(FAssetIdentifier(NAME_None));
-		RecursivelyPopulateNodeInfos(false, GraphRootIdentifiers[0], DependencyNodeInfos, 0, Settings->GetSearchDependencyDepthLimit());
-
-		TMap<FAssetIdentifier, int32> ReferencerNodeSizes;
-		TSet<FAssetIdentifier> VisitedReferencerSizeNames;
-		if (Settings->IsShowReferencers())
-		{
-			const int32 ReferencerDepth = 1;
-			RecursivelyGatherSizes(/*bReferencers=*/true, GraphRootIdentifiers, CurrentCollectionPackages, ReferencerDepth, Settings->GetSearchReferencerDepthLimit(), VisitedReferencerSizeNames, ReferencerNodeSizes);
-		}
-
-		TMap<FAssetIdentifier, int32> DependencyNodeSizes;
-		TSet<FAssetIdentifier> VisitedDependencySizeNames;
-		if (Settings->IsShowDependencies())
-		{
-			const int32 DependencyDepth = 1;
-			RecursivelyGatherSizes(/*bReferencers=*/false, GraphRootIdentifiers, CurrentCollectionPackages, DependencyDepth, Settings->GetSearchDependencyDepthLimit(), VisitedDependencySizeNames, DependencyNodeSizes);
-		}
+		RecursivelyPopulateNodeInfos(false, GraphRootIdentifiers[0], NewDependencyNodeInfos, 0, Settings->GetSearchDependencyDepthLimit());
 
 		TSet<FName> AllPackageNames;
 		auto AddPackage = [](const FAssetIdentifier& AssetId, TSet<FName>& PackageNames)
@@ -227,71 +214,147 @@ UEdGraphNode_Reference* UEdGraph_ReferenceViewer::ConstructNodes(const TArray<FA
 
 		if (Settings->IsShowReferencers())
 		{
-			for (const FAssetIdentifier& AssetId : VisitedReferencerSizeNames)
+			for (TPair<FAssetIdentifier, FReferenceNodeInfo>&  InfoPair : NewReferenceNodeInfos)
 			{
-				AddPackage(AssetId, AllPackageNames);
+				AddPackage(InfoPair.Key, AllPackageNames);
 			}
 		}
 
 		if (Settings->IsShowDependencies())
 		{
-			for (const FAssetIdentifier& AssetId : VisitedDependencySizeNames)
+			for (TPair<FAssetIdentifier, FReferenceNodeInfo>&  InfoPair : NewDependencyNodeInfos)
 			{
-				AddPackage(AssetId, AllPackageNames);
+				AddPackage(InfoPair.Key, AllPackageNames);
 			}
 		}
 
+		// Store the AssetData in the NodeInfos
 		TMap<FName, FAssetData> PackagesToAssetDataMap;
 		GatherAssetData(AllPackageNames, PackagesToAssetDataMap);
 
-		// Store the AssetData in the NodeInfos
-		for (TPair<FAssetIdentifier, ReferenceNodeInfo>&  InfoPair : ReferenceNodeInfos)
+		// Store the AssetData in the NodeInfos and collect Asset Type UClasses to populate the filters
+		TSet<FName> AllClasses;
+		for (TPair<FAssetIdentifier, FReferenceNodeInfo>&  InfoPair : NewReferenceNodeInfos)
 		{
 			InfoPair.Value.AssetData = PackagesToAssetDataMap.FindRef(InfoPair.Key.PackageName);
+			if (InfoPair.Value.AssetData.IsValid())
+			{
+				AllClasses.Add(InfoPair.Value.AssetData.AssetClass);
+			}
 		}
 
-		for (TPair<FAssetIdentifier, ReferenceNodeInfo>&  InfoPair : DependencyNodeInfos)
+		for (TPair<FAssetIdentifier, FReferenceNodeInfo>&  InfoPair : NewDependencyNodeInfos)
 		{
 			InfoPair.Value.AssetData = PackagesToAssetDataMap.FindRef(InfoPair.Key.PackageName);
+			if (InfoPair.Value.AssetData.IsValid())
+			{
+				AllClasses.Add(InfoPair.Value.AssetData.AssetClass);
+			}
 		}
 
+		// Update the cached class types list
+		CurrentClasses = AllClasses;
+		OnAssetsChangedDelegate.Broadcast();
+
+		ReferencerNodeInfos = NewReferenceNodeInfos;
+		DependencyNodeInfos = NewDependencyNodeInfos;
+	}
+	else
+	{
+		ReferencerNodeInfos.Empty();
+		DependencyNodeInfos.Empty();
+	}
+
+	return RefilterGraph();
+}
+
+
+UEdGraphNode_Reference* UEdGraph_ReferenceViewer::RefilterGraph()
+{
+	RemoveAllNodes();
+	UEdGraphNode_Reference* RootNode = NULL;
+	if (CurrentGraphRootIdentifiers.Num() > 0)
+	{
+		FAssetIdentifier FirstGraphRootIdentifier = CurrentGraphRootIdentifiers[0];
+
 		// Create the root node
+		bool bRootIsDuplicated = (Settings->IsShowDependencies() && DependencyNodeInfos[FirstGraphRootIdentifier].IsADuplicate()) || 
+								  (Settings->IsShowReferencers() && ReferencerNodeInfos[FirstGraphRootIdentifier].IsADuplicate());
+		const FReferenceNodeInfo& NodeInfo = Settings->IsShowReferencers() ? ReferencerNodeInfos[FirstGraphRootIdentifier] : DependencyNodeInfos[FirstGraphRootIdentifier];
 		RootNode = CreateReferenceNode();
-		bool bRootIsDuplicated = DependencyNodeInfos[GraphRootIdentifiers[0]].IsADuplicate() || ReferenceNodeInfos[GraphRootIdentifiers[0]].IsADuplicate();
-		RootNode->SetupReferenceNode(GraphRootOrigin, GraphRootIdentifiers, PackagesToAssetDataMap.FindRef(GraphRootIdentifiers[0].PackageName), /*bInAllowThumbnail = */ !Settings->IsCompactMode(), /*bIsDuplicate*/ bRootIsDuplicated);
+		RootNode->SetupReferenceNode(CurrentGraphRootOrigin, CurrentGraphRootIdentifiers, NodeInfo.AssetData, /*bInAllowThumbnail = */ !Settings->IsCompactMode(), /*bIsDuplicate*/ bRootIsDuplicated);
 		RootNode->SetMakeCommentBubbleVisible(Settings->IsShowPath());
 
-		if (bUseNodeInfos) // @LULU true = using new stuff, false = Not using new stuff 
+		if (Settings->IsShowReferencers())
 		{
-			if (Settings->IsShowReferencers())
-			{
-				RecursivelyCreateNodes(true, GraphRootIdentifiers[0], GraphRootOrigin, GraphRootIdentifiers[0], RootNode, ReferenceNodeInfos, 0, Settings->GetSearchReferencerDepthLimit(), /*bIsRoot*/ true);
-			}
+			RecursivelyFilterNodeInfos(FirstGraphRootIdentifier, ReferencerNodeInfos, 0, Settings->GetSearchReferencerDepthLimit());
+			RecursivelyCreateNodes(true, FirstGraphRootIdentifier, CurrentGraphRootOrigin, FirstGraphRootIdentifier, RootNode, ReferencerNodeInfos, 0, Settings->GetSearchReferencerDepthLimit(), /*bIsRoot*/ true);
+		}
 
-			if (Settings->IsShowDependencies())
-			{
-				RecursivelyCreateNodes(false, GraphRootIdentifiers[0], GraphRootOrigin, GraphRootIdentifiers[0], RootNode, DependencyNodeInfos, 0, Settings->GetSearchDependencyDepthLimit(), /*bIsRoot*/ true);
-			}
+		if (Settings->IsShowDependencies())
+		{
+			RecursivelyFilterNodeInfos(FirstGraphRootIdentifier, DependencyNodeInfos, 0, Settings->GetSearchDependencyDepthLimit());
+			RecursivelyCreateNodes(false, FirstGraphRootIdentifier, CurrentGraphRootOrigin, FirstGraphRootIdentifier, RootNode, DependencyNodeInfos, 0, Settings->GetSearchDependencyDepthLimit(), /*bIsRoot*/ true);
 		}
-		else {
 
-			if (Settings->IsShowReferencers())
+	}
+
+	NotifyGraphChanged();
+	return RootNode;
+}
+
+void UEdGraph_ReferenceViewer::RecursivelyFilterNodeInfos(const FAssetIdentifier& InAssetId, TMap<FAssetIdentifier, FReferenceNodeInfo>& InNodeInfos, int32 InCurrentDepth, int32 InMaxDepth)
+{
+	// Filters and Reprovisions the NodeInfo counts 
+	int32 NewProvisionSize = 0;
+
+	int32 Breadth = 0;
+
+	InNodeInfos[InAssetId].OverflowCount = 0;
+	if (InMaxDepth > 0 && InCurrentDepth < InMaxDepth)
+	{
+		for (const TPair<FAssetIdentifier, EDependencyPinCategory>& Pair : InNodeInfos[InAssetId].Children)
+		{
+			FAssetIdentifier ChildId = Pair.Key;
+
+			int32 ChildProvSize = 0;
+			if (InNodeInfos[ChildId].IsFirstParent(InAssetId))
+			{
+				RecursivelyFilterNodeInfos(ChildId, InNodeInfos, InCurrentDepth + 1, InMaxDepth);
+				ChildProvSize = InNodeInfos[ChildId].ProvisionSize(InAssetId);
+			}
+			else if (InNodeInfos[ChildId].PassedFilters && Settings->IsShowDuplicates())
 			{
-				TSet<FAssetIdentifier> VisitedReferencerNames;
-				const int32 VisitedReferencerDepth = 1;
-				RecursivelyConstructNodes(/*bReferencers=*/true, RootNode, GraphRootIdentifiers, GraphRootOrigin, ReferencerNodeSizes, PackagesToAssetDataMap, CurrentCollectionPackages, VisitedReferencerDepth, Settings->GetSearchReferencerDepthLimit(), VisitedReferencerNames);
+				ChildProvSize = 1;
 			}
 
-			if (Settings->IsShowDependencies())
+			if (ChildProvSize > 0)
 			{
-				TSet<FAssetIdentifier> VisitedDependencyNames;
-				const int32 VisitedDependencyDepth = 1;
-				RecursivelyConstructNodes(/*bReferencers=*/false, RootNode, GraphRootIdentifiers, GraphRootOrigin, DependencyNodeSizes, PackagesToAssetDataMap, CurrentCollectionPackages, VisitedDependencyDepth, Settings->GetSearchDependencyDepthLimit(), VisitedDependencyNames);
+				if (!ExceedsMaxSearchBreadth(Breadth))
+				{
+					NewProvisionSize += ChildProvSize;
+					Breadth++;
+				}
+
+				else
+				{
+					InNodeInfos[InAssetId].OverflowCount++;
+					Breadth++;
+				}
 			}
 		}
 	}
 
-	return RootNode;
+	// Account for an overflow node if necessary
+	if (InNodeInfos[InAssetId].OverflowCount > 0)
+	{
+		NewProvisionSize++;
+	}
+
+	bool PassedAllFilters = FilterCollection && Settings->GetFiltersEnabled() ? FilterCollection->PassesAllFilters(InNodeInfos[InAssetId]) : true;
+
+	InNodeInfos[InAssetId].ChildProvisionSize = NewProvisionSize > 0 ? NewProvisionSize : (PassedAllFilters ? 1 : 0);
+	InNodeInfos[InAssetId].PassedFilters = PassedAllFilters;
 }
 
 void UEdGraph_ReferenceViewer::GetSortedLinks(const TArray<FAssetIdentifier>& Identifiers, bool bReferencers, const FAssetManagerDependencyQuery& Query, TMap<FAssetIdentifier, EDependencyPinCategory>& OutLinks) const
@@ -400,11 +463,10 @@ bool UEdGraph_ReferenceViewer::IsPackageIdentifierPassingFilter(const FAssetIden
 }
 
 void
-UEdGraph_ReferenceViewer::RecursivelyPopulateNodeInfos(bool bInReferencers, const FAssetIdentifier& InAssetId, TMap<FAssetIdentifier, ReferenceNodeInfo>& InNodeInfos, int32 InCurrentDepth, int32 InMaxDepth)
+UEdGraph_ReferenceViewer::RecursivelyPopulateNodeInfos(bool bInReferencers, const FAssetIdentifier& InAssetId, TMap<FAssetIdentifier, FReferenceNodeInfo>& InNodeInfos, int32 InCurrentDepth, int32 InMaxDepth)
 {
 	int32 ProvisionSize = 0;
 
-	int32 Breadth = 0;
 	if (InMaxDepth > 0 && InCurrentDepth < InMaxDepth)
 	{
 		TArray<FAssetIdentifier> AssetIds;
@@ -417,38 +479,22 @@ UEdGraph_ReferenceViewer::RecursivelyPopulateNodeInfos(bool bInReferencers, cons
 		for (const TPair<FAssetIdentifier, EDependencyPinCategory>& Pair : ReferenceLinks)
 		{
 			FAssetIdentifier ChildId = Pair.Key;
-			if (!ExceedsMaxSearchBreadth(Breadth))
+			if (!InNodeInfos.Contains(ChildId))
 			{
-				// Only Gather Children the first time
-				if (!InNodeInfos.Contains(ChildId))
-				{
-					ReferenceNodeInfo& NewNodeInfo = InNodeInfos.FindOrAdd(ChildId, ReferenceNodeInfo(ChildId, bInReferencers));
-					InNodeInfos[ChildId].Parents.Emplace(InAssetId);
-					InNodeInfos[InAssetId].Children.Emplace(Pair);
+				FReferenceNodeInfo& NewNodeInfo = InNodeInfos.FindOrAdd(ChildId, FReferenceNodeInfo(ChildId, bInReferencers));
+				InNodeInfos[ChildId].Parents.Emplace(InAssetId);
+				InNodeInfos[InAssetId].Children.Emplace(Pair);
 
-					RecursivelyPopulateNodeInfos(bInReferencers, ChildId, InNodeInfos, InCurrentDepth + 1, InMaxDepth);
-					ProvisionSize += InNodeInfos[ChildId].ProvisionSize(InAssetId);
-					Breadth++;
-				}
-
-				else if (Settings->IsShowDuplicates() && !InNodeInfos[ChildId].Parents.Contains(InAssetId))
-				{
-					InNodeInfos[ChildId].Parents.Emplace(InAssetId);
-					InNodeInfos[InAssetId].Children.Emplace(Pair);
-					ProvisionSize += 1;
-					Breadth++;
-				}
+				RecursivelyPopulateNodeInfos(bInReferencers, ChildId, InNodeInfos, InCurrentDepth + 1, InMaxDepth);
+				ProvisionSize += InNodeInfos[ChildId].ProvisionSize(InAssetId);
 			}
-			else
+
+			else if (!InNodeInfos[ChildId].Parents.Contains(InAssetId))
 			{
-				// Count the overflow nodes to report in the UI but otherwise skip adding them 
-				if (Settings->IsShowDuplicates() || !InNodeInfos.Contains(ChildId))
-				{
-					InNodeInfos[InAssetId].OverflowCount++;
-					Breadth++;
-				}
+				InNodeInfos[ChildId].Parents.Emplace(InAssetId);
+				InNodeInfos[InAssetId].Children.Emplace(Pair);
+				ProvisionSize += 1;
 			}
-
 		}
 	}
 
@@ -461,68 +507,16 @@ UEdGraph_ReferenceViewer::RecursivelyPopulateNodeInfos(bool bInReferencers, cons
 	InNodeInfos[InAssetId].ChildProvisionSize = ProvisionSize > 0 ? ProvisionSize : 1;
 }
 
-int32 UEdGraph_ReferenceViewer::RecursivelyGatherSizes(bool bReferencers, const TArray<FAssetIdentifier>& Identifiers, const TSet<FName>& AllowedPackageNames, int32 CurrentDepth, int32 MaxDepth, TSet<FAssetIdentifier>& VisitedNames, TMap<FAssetIdentifier, int32>& OutNodeSizes) const
-{
-	check(Identifiers.Num() > 0);
-
-	VisitedNames.Append(Identifiers);
-
-	TMap<FAssetIdentifier, EDependencyPinCategory> ReferenceLinks;
-	GetSortedLinks(Identifiers, bReferencers, GetReferenceSearchFlags(false), ReferenceLinks);
-
-	int32 NodeSize = 0;
-
-	if ( ReferenceLinks.Num() > 0 && !ExceedsMaxSearchDepth(CurrentDepth, MaxDepth) )
-	{
-		int32 NumReferencesMade = 0;
-		int32 NumReferencesExceedingMax = 0;
-
-		for (const TPair<FAssetIdentifier, EDependencyPinCategory>& Pair : ReferenceLinks)
-		{
-			const FAssetIdentifier& AssetId = Pair.Key;
-			if ( !VisitedNames.Contains(AssetId) )
-			{
-				if ( !ExceedsMaxSearchBreadth(NumReferencesMade) )
-				{
-					TArray<FAssetIdentifier> NewPackageNames;
-					NewPackageNames.Add(AssetId);
-					NodeSize += RecursivelyGatherSizes(bReferencers, NewPackageNames, AllowedPackageNames, CurrentDepth + 1, MaxDepth, VisitedNames, OutNodeSizes);
-					NumReferencesMade++;
-				}
-				else
-				{
-					NumReferencesExceedingMax++;
-				}
-			}
-		}
-
-		if ( NumReferencesExceedingMax > 0 )
-		{
-			// Add one size for the collapsed node
-			NodeSize++;
-		}
-	}
-
-	if ( NodeSize == 0 )
-	{
-		// If you have no valid children, the node size is just 1 (counting only self to make a straight line)
-		NodeSize = 1;
-	}
-
-	OutNodeSizes.Add(Identifiers[0], NodeSize);
-	return NodeSize;
-}
-
 void UEdGraph_ReferenceViewer::GatherAssetData(const TSet<FName>& AllPackageNames, TMap<FName, FAssetData>& OutPackageToAssetDataMap) const
 {
 	UE::AssetRegistry::GetAssetForPackages(AllPackageNames.Array(), OutPackageToAssetDataMap);
 }
 
-UEdGraphNode_Reference* UEdGraph_ReferenceViewer::RecursivelyCreateNodes(bool bInReferencers, const FAssetIdentifier& InAssetId, const FIntPoint& InNodeLoc, const FAssetIdentifier& InParentId, UEdGraphNode_Reference* InParentNode, TMap<FAssetIdentifier, ReferenceNodeInfo>& InNodeInfos, int32 InCurrentDepth, int32 InMaxDepth, bool bIsRoot)
+UEdGraphNode_Reference* UEdGraph_ReferenceViewer::RecursivelyCreateNodes(bool bInReferencers, const FAssetIdentifier& InAssetId, const FIntPoint& InNodeLoc, const FAssetIdentifier& InParentId, UEdGraphNode_Reference* InParentNode, TMap<FAssetIdentifier, FReferenceNodeInfo>& InNodeInfos, int32 InCurrentDepth, int32 InMaxDepth, bool bIsRoot)
 {
 	check(InNodeInfos.Contains(InAssetId));
 
-	const ReferenceNodeInfo& NodeInfo = InNodeInfos[InAssetId];
+	const FReferenceNodeInfo& NodeInfo = InNodeInfos[InAssetId];
 	int32 NodeProvSize = 1;
 
 	UEdGraphNode_Reference* NewNode = NULL;
@@ -534,8 +528,9 @@ UEdGraphNode_Reference* UEdGraph_ReferenceViewer::RecursivelyCreateNodes(bool bI
 	else
 	{
 		NewNode = CreateReferenceNode();
-		NewNode->SetupReferenceNode(InNodeLoc, {InAssetId}, NodeInfo.AssetData, /*bInAllowThumbnail*/ !Settings->IsCompactMode(), /*bIsADuplicate*/ NodeInfo.Parents.Num() > 1);
+		NewNode->SetupReferenceNode(InNodeLoc, {InAssetId}, NodeInfo.AssetData, /*bInAllowThumbnail*/ !Settings->IsCompactMode() && NodeInfo.PassedFilters, /*bIsADuplicate*/ NodeInfo.Parents.Num() > 1);
 		NewNode->SetMakeCommentBubbleVisible(Settings->IsShowPath());
+		NewNode->SetIsFiltered(!NodeInfo.PassedFilters);
 		NodeProvSize = NodeInfo.ProvisionSize(InParentId);
 	}
 
@@ -557,10 +552,25 @@ UEdGraphNode_Reference* UEdGraph_ReferenceViewer::RecursivelyCreateNodes(bool bI
 
 		for (const TPair<FAssetIdentifier, EDependencyPinCategory>& Pair : InNodeInfos[InAssetId].Children)
 		{
+			if (ExceedsMaxSearchBreadth(Breadth))
+			{
+				break;
+			}
 
-			    FAssetIdentifier ChildId = Pair.Key;
-			    int32 ChildProvSize = InNodeInfos[ChildId].ProvisionSize(InAssetId);
-
+		    FAssetIdentifier ChildId = Pair.Key;
+		    int32 ChildProvSize = 0;
+		    if (InNodeInfos[ChildId].IsFirstParent(InAssetId))
+		   	{
+		   		ChildProvSize = InNodeInfos[ChildId].ProvisionSize(InAssetId);
+		   	}
+		   	else if (InNodeInfos[ChildId].PassedFilters && Settings->IsShowDuplicates())
+		   	{
+		   		ChildProvSize = 1;
+		   	}
+
+		    // The provision size will always be at least 1 if it should be shown, factoring in filters, breadth, duplicates, etc.
+		   	if (ChildProvSize > 0)
+		    {
 				ChildLoc.Y += (ChildProvSize - 1) * NodeSizeY * 0.5;
 
 				UEdGraphNode_Reference* ChildNode = RecursivelyCreateNodes(bInReferencers, ChildId, ChildLoc, InAssetId, NewNode, InNodeInfos, InCurrentDepth + 1, InMaxDepth);	
@@ -577,7 +587,8 @@ UEdGraphNode_Reference* UEdGraph_ReferenceViewer::RecursivelyCreateNodes(bool bI
 				}
 
 				ChildLoc.Y += NodeSizeY * (ChildProvSize + 1) * 0.5;
-
+				Breadth ++;
+		    }
 		}
 
 		// There were more references than allowed to be displayed. Make a collapsed node.
@@ -608,130 +619,6 @@ UEdGraphNode_Reference* UEdGraph_ReferenceViewer::RecursivelyCreateNodes(bool bI
 	return NewNode;
 }
 
-UEdGraphNode_Reference* UEdGraph_ReferenceViewer::RecursivelyConstructNodes(bool bReferencers, UEdGraphNode_Reference* RootNode, const TArray<FAssetIdentifier>& Identifiers, const FIntPoint& NodeLoc, const TMap<FAssetIdentifier, int32>& NodeSizes, const TMap<FName, FAssetData>& PackagesToAssetDataMap, const TSet<FName>& AllowedPackageNames, int32 CurrentDepth, int32 MaxDepth, TSet<FAssetIdentifier>& VisitedNames)
-{
-	check(Identifiers.Num() > 0);
-
-	VisitedNames.Append(Identifiers);
-
-	UEdGraphNode_Reference* NewNode = NULL;
-	if ( RootNode->GetIdentifier() == Identifiers[0] )
-	{
-		// Don't create the root node. It is already created!
-		NewNode = RootNode;
-	}
-	else
-	{
-		NewNode = CreateReferenceNode();
-		NewNode->SetupReferenceNode(NodeLoc, Identifiers, PackagesToAssetDataMap.FindRef(Identifiers[0].PackageName), /*bInAllowThumbnail = */ !Settings->IsCompactMode(), /*bIsDuplicate*/ false);
-	}
-
-	TMap<FAssetIdentifier, EDependencyPinCategory> Referencers;
-	GetSortedLinks(Identifiers, bReferencers, GetReferenceSearchFlags(false), Referencers);
-
-	if (Referencers.Num() > 0 && !ExceedsMaxSearchDepth(CurrentDepth, MaxDepth))
-	{
-		FIntPoint ReferenceNodeLoc = NodeLoc;
-
-		if (bReferencers)
-		{
-			// Referencers go left
-			ReferenceNodeLoc.X -= Settings->IsCompactMode() ? 400 : 800;
-		}
-		else
-		{
-			// Dependencies go right
-			ReferenceNodeLoc.X += Settings->IsCompactMode() ? 400 : 800;
-		}
-
-		const int32 NodeSizeY = Settings->IsCompactMode() ? 100 : 200;
-		const int32 TotalReferenceSizeY = NodeSizes.FindChecked(Identifiers[0]) * NodeSizeY;
-
-		ReferenceNodeLoc.Y -= TotalReferenceSizeY * 0.5f;
-		ReferenceNodeLoc.Y += NodeSizeY * 0.5f;
-
-		int32 NumReferencesMade = 0;
-		int32 NumReferencesExceedingMax = 0;
-
-		for ( TPair<FAssetIdentifier, EDependencyPinCategory>& DependencyPair : Referencers)
-		{
-			FAssetIdentifier ReferenceName = DependencyPair.Key;
-
-			if ( !VisitedNames.Contains(ReferenceName) )
-			{
-				if ( !ExceedsMaxSearchBreadth(NumReferencesMade) )
-				{
-					int32 ThisNodeSizeY = ReferenceName.IsValue() ? 100 : NodeSizeY;
-
-					const int32 RefSizeY = NodeSizes.FindChecked(ReferenceName);
-					FIntPoint RefNodeLoc;
-					RefNodeLoc.X = ReferenceNodeLoc.X;
-					RefNodeLoc.Y = ReferenceNodeLoc.Y + RefSizeY * ThisNodeSizeY * 0.5 - ThisNodeSizeY * 0.5;
-					
-					TArray<FAssetIdentifier> NewIdentifiers;
-					NewIdentifiers.Add(ReferenceName);
-					
-					UEdGraphNode_Reference* ReferenceNode = RecursivelyConstructNodes(bReferencers, RootNode, NewIdentifiers, RefNodeLoc, NodeSizes, PackagesToAssetDataMap, AllowedPackageNames, CurrentDepth + 1, MaxDepth, VisitedNames);
-					if ( ensure(ReferenceNode) )
-					{
-						if (bReferencers)
-						{
-							ReferenceNode->GetDependencyPin()->PinType.PinCategory = ::GetName(DependencyPair.Value);
-						}
-						else
-						{
-							ReferenceNode->GetReferencerPin()->PinType.PinCategory = ::GetName(DependencyPair.Value);
-						}
-
-						if ( bReferencers )
-						{
-							NewNode->AddReferencer( ReferenceNode );
-						}
-						else
-						{
-							ReferenceNode->AddReferencer( NewNode );
-						}
-
-						ReferenceNodeLoc.Y += RefSizeY * ThisNodeSizeY;
-					}
-
-					NumReferencesMade++;
-				}
-				else
-				{
-					NumReferencesExceedingMax++;
-				}
-			}
-		}
-
-		if ( NumReferencesExceedingMax > 0 )
-		{
-			// There are more references than allowed to be displayed. Make a collapsed node.
-			UEdGraphNode_Reference* ReferenceNode = CreateReferenceNode();
-			FIntPoint RefNodeLoc;
-			RefNodeLoc.X = ReferenceNodeLoc.X;
-			RefNodeLoc.Y = ReferenceNodeLoc.Y;
-
-			if ( ensure(ReferenceNode) )
-			{
-				ReferenceNode->SetAllowThumbnail(!Settings->IsCompactMode());
-				ReferenceNode->SetReferenceNodeCollapsed(RefNodeLoc, NumReferencesExceedingMax);
-
-				if ( bReferencers )
-				{
-					NewNode->AddReferencer( ReferenceNode );
-				}
-				else
-				{
-					ReferenceNode->AddReferencer( NewNode );
-				}
-			}
-		}
-	}
-
-	return NewNode;
-}
-
 const TSharedPtr<FAssetThumbnailPool>& UEdGraph_ReferenceViewer::GetAssetThumbnailPool() const
 {
 	return AssetThumbnailPool;
diff --git a/Engine/Plugins/Editor/AssetManagerEditor/Source/AssetManagerEditor/Private/ReferenceViewer/ReferenceViewerSettings.cpp b/Engine/Plugins/Editor/AssetManagerEditor/Source/AssetManagerEditor/Private/ReferenceViewer/ReferenceViewerSettings.cpp
--- a/Engine/Plugins/Editor/AssetManagerEditor/Source/AssetManagerEditor/Private/ReferenceViewer/ReferenceViewerSettings.cpp
+++ b/Engine/Plugins/Editor/AssetManagerEditor/Source/AssetManagerEditor/Private/ReferenceViewer/ReferenceViewerSettings.cpp
@@ -200,3 +200,35 @@ void UReferenceViewerSettings::SetShowPathEnabled(bool bEnabled)
 	SaveConfig();
 }
 
+bool UReferenceViewerSettings::GetFiltersEnabled() const
+{
+	return bFiltersEnabled;
+}
+
+void UReferenceViewerSettings::SetFiltersEnabled(bool newEnabled)
+{
+	bFiltersEnabled = newEnabled;
+	SaveConfig();
+}
+
+bool UReferenceViewerSettings::AutoUpdateFilters() const
+{
+	return bAutoUpdateFilters;
+}
+
+void UReferenceViewerSettings::SetAutoUpdateFilters(bool bEnabled)
+{
+	bAutoUpdateFilters = bEnabled;
+	SaveConfig();
+}
+
+const TArray<FilterState>& UReferenceViewerSettings::GetUserFilters() const
+{
+	return UserFilters;	
+}
+
+void UReferenceViewerSettings::SetUserFilters(TArray<FilterState>& InFilters)
+{
+	UserFilters = InFilters;
+	SaveConfig();
+}
diff --git a/Engine/Plugins/Editor/AssetManagerEditor/Source/AssetManagerEditor/Private/ReferenceViewer/ReferenceViewerStyle.cpp b/Engine/Plugins/Editor/AssetManagerEditor/Source/AssetManagerEditor/Private/ReferenceViewer/ReferenceViewerStyle.cpp
--- a/Engine/Plugins/Editor/AssetManagerEditor/Source/AssetManagerEditor/Private/ReferenceViewer/ReferenceViewerStyle.cpp
+++ b/Engine/Plugins/Editor/AssetManagerEditor/Source/AssetManagerEditor/Private/ReferenceViewer/ReferenceViewerStyle.cpp
@@ -75,7 +75,17 @@ FReferenceViewerStyle::FReferenceViewerStyle()
 	Set( "Graph.Node.ColorSpill",     new FSlateRoundedBoxBrush(SpillColor, FVector4(BodyRadius, BodyRadius, 0.0, 0.0)));
 
 	Set( "Graph.Node.Duplicate",      new IMAGE_BRUSH_SVG("/GraphNode_Duplicate_8px", FVector2D(8.f, 8.f), FStyleColors::White));
-	
+
+
+	const FVector2D IconSize(20.0f, 20.0f);
+
+	Set("Icons.ArrowLeft", new CORE_IMAGE_BRUSH_SVG("Starship/Common/arrow-left", IconSize));
+	Set("Icons.ArrowRight", new CORE_IMAGE_BRUSH_SVG("Starship/Common/arrow-right", IconSize));
+
+	Set("Icons.AutoFilters", new CORE_IMAGE_BRUSH_SVG("Starship/Common/FilterAuto", IconSize));
+	Set("Icons.Filters", new CORE_IMAGE_BRUSH_SVG("Starship/Common/filter", IconSize));
+	Set("Icons.Duplicate",      new IMAGE_BRUSH_SVG("/GraphNode_Duplicate_8px", IconSize, FStyleColors::White));
+
 	FSlateStyleRegistry::RegisterSlateStyle(*this);
 }
 
diff --git a/Engine/Plugins/Editor/AssetManagerEditor/Source/AssetManagerEditor/Private/ReferenceViewer/SReferenceNode.cpp b/Engine/Plugins/Editor/AssetManagerEditor/Source/AssetManagerEditor/Private/ReferenceViewer/SReferenceNode.cpp
--- a/Engine/Plugins/Editor/AssetManagerEditor/Source/AssetManagerEditor/Private/ReferenceViewer/SReferenceNode.cpp
+++ b/Engine/Plugins/Editor/AssetManagerEditor/Source/AssetManagerEditor/Private/ReferenceViewer/SReferenceNode.cpp
@@ -75,6 +75,9 @@ void SReferenceNode::UpdateGraphNode()
 	TSharedRef<SWidget> ThumbnailWidget = SNullWidget::NullWidget;
 	UEdGraphNode_Reference* RefGraphNode = CastChecked<UEdGraphNode_Reference>(GraphNode);
 	bool bIsADuplicate = RefGraphNode->IsADuplicate();
+
+	FLinearColor OpacityColor = RefGraphNode->GetIsFiltered() ? FLinearColor(1.0, 1.0, 1.0, 0.4) : FLinearColor::White;
+	
 	if ( AssetThumbnail.IsValid() )
 	{
 
@@ -101,6 +104,7 @@ void SReferenceNode::UpdateGraphNode()
 		.AutoHeight()
 		[
 			SNew(SBorder)
+			.ColorAndOpacity(OpacityColor)
 			.BorderImage( FReferenceViewerStyle::Get().GetBrush( "Graph.Node.BodyBackground" ) )
 			.Padding(0)
 			[
@@ -169,7 +173,7 @@ void SReferenceNode::UpdateGraphNode()
 						.AutoWidth()
 						.HAlign(HAlign_Right)
 						.VAlign(VAlign_Bottom)
-						.Padding(FMargin(0.f, 0.f, 0.f, 1.f))
+						.Padding(FMargin(4.f, 0.f, 0.f, 1.f))
 						[
 							SNew(SImage)
 							.Visibility(bIsADuplicate ? EVisibility::Visible : EVisibility::Hidden)
diff --git a/Engine/Plugins/Editor/AssetManagerEditor/Source/AssetManagerEditor/Private/ReferenceViewer/SReferenceViewer.cpp b/Engine/Plugins/Editor/AssetManagerEditor/Source/AssetManagerEditor/Private/ReferenceViewer/SReferenceViewer.cpp
--- a/Engine/Plugins/Editor/AssetManagerEditor/Source/AssetManagerEditor/Private/ReferenceViewer/SReferenceViewer.cpp
+++ b/Engine/Plugins/Editor/AssetManagerEditor/Source/AssetManagerEditor/Private/ReferenceViewer/SReferenceViewer.cpp
@@ -21,8 +21,8 @@
 #include "SSimpleComboButton.h"
 
 #include "Styling/AppStyle.h"
+#include "ReferenceViewerStyle.h"
 #include "Engine/Selection.h"
-#include "ReferenceViewer/EdGraph_ReferenceViewer.h"
 #include "ReferenceViewer/EdGraphNode_Reference.h"
 #include "ReferenceViewer/ReferenceViewerSchema.h"
 #include "AssetRegistry/AssetRegistryModule.h"
@@ -43,12 +43,6 @@
 
 #define LOCTEXT_NAMESPACE "ReferenceViewer"
 
-static bool GShowToggleDeprecatedReferenceViewerLayout = false;
-static FAutoConsoleCommand CVarShowToggleDeprecatedReferenceViewerLayout(
-	TEXT("ReferenceViewer.ShowToggleDeprecatedLayout"),
-	TEXT("Displays the toggle allowing the user to toggle back to the former layout algorithm."),
-	FConsoleCommandDelegate::CreateLambda([] { GShowToggleDeprecatedReferenceViewerLayout = !GShowToggleDeprecatedReferenceViewerLayout; }));
-
 bool IsPackageNamePassingFilter(FName InPackageName, const TArray<FString>& InSearchWords)
 {
 	// package name must match all words
@@ -97,6 +91,7 @@ SReferenceViewer::~SReferenceViewer()
 
 void SReferenceViewer::Construct(const FArguments& InArgs)
 {
+	bRebuildingFilters = false;
 	Settings = GetMutableDefault<UReferenceViewerSettings>();
 
 	// Create an action list and register commands
@@ -111,6 +106,7 @@ void SReferenceViewer::Construct(const FArguments& InArgs)
 	GraphObj->Schema = UReferenceViewerSchema::StaticClass();
 	GraphObj->AddToRoot();
 	GraphObj->SetReferenceViewer(StaticCastSharedRef<SReferenceViewer>(AsShared()));
+	GraphObj->OnAssetsChanged().AddSP(this, &SReferenceViewer::OnUpdateFilterBar);
 
 	SGraphEditor::FGraphEditorEvents GraphEvents;
 	GraphEvents.OnNodeDoubleClicked = FSingleNodeEvent::CreateSP(this, &SReferenceViewer::OnNodeDoubleClicked);
@@ -143,6 +139,26 @@ void SReferenceViewer::Construct(const FArguments& InArgs)
 	bShowCompactMode = true;
 	bDirtyResults = false;
 
+	SAssignNew(FilterWidget, SReferenceViewerFilterBar)
+		.Visibility_Lambda([this]() { return !Settings->GetFiltersEnabled() ? EVisibility::Collapsed : EVisibility::Visible; })
+		.OnConvertItemToAssetData_Lambda([this] (FReferenceNodeInfo& InNodeInfo, FAssetData& OutAssetData) -> bool { 
+			OutAssetData = InNodeInfo.AssetData; 
+			return true; 
+		})
+		.UseDefaultAssetFilters(true)
+		.OnFilterChanged_Lambda([this] { 
+			if (!bRebuildingFilters && GraphObj)
+			{
+				GraphObj->SetCurrentFilterCollection(FilterWidget->GetAllActiveFilters());
+				GraphObj->RefilterGraph();
+				FilterWidget->SaveSettings(FString(), FString(), FString());
+			}
+		})
+	;
+
+	TSharedPtr<SWidget> FilterCombo = FilterWidget->MakeAddFilterButton(FilterWidget.ToSharedRef());
+	FilterCombo->SetVisibility(TAttribute<EVisibility>::CreateLambda([this]() { return !Settings->GetFiltersEnabled() ? EVisibility::Collapsed : EVisibility::Visible; }));
+
 	ChildSlot
 	[
 
@@ -154,75 +170,21 @@ void SReferenceViewer::Construct(const FArguments& InArgs)
 		[
 			SNew(SBorder)
 			.BorderImage(FAppStyle::GetBrush("Brushes.Panel"))
-			.Padding(FMargin(12, 6, 12, 6))
 			[
 				SNew(SHorizontalBox)
 
-
-				// Refresh Button
-				+SHorizontalBox::Slot()
-				.AutoWidth()
-				.Padding(4, 0)
-				[
-					SNew(SSimpleButton)
-					.ToolTipText(LOCTEXT("RefreshTooltip", "Refresh current view"))
-					.OnClicked(this, &SReferenceViewer::RefreshClicked)
-					.Icon(FAppStyle::GetBrush("Icons.Refresh"))
-				]
-
-
-				// History Back Button
-				+SHorizontalBox::Slot()
-				.AutoWidth()
-				.Padding(4, 0)
-				[
-					SNew(SButton)
-					.ToolTipText( this, &SReferenceViewer::GetHistoryBackTooltip )
-					.ButtonStyle(&FAppStyle::Get().GetWidgetStyle<FButtonStyle>("SimpleButton"))
-					.ForegroundColor(FSlateColor::UseStyle())
-					.OnClicked(this, &SReferenceViewer::BackClicked)
-					.IsEnabled(this, &SReferenceViewer::IsBackEnabled)
-					[
-						SNew(SImage)
-						.ColorAndOpacity(FSlateColor::UseForeground())
-						.Image(FAppStyle::GetBrush("Icons.ArrowLeft"))
-						.DesiredSizeOverride(FVector2D(20.0f, 20.0f))
-					]
-				]
-
-				// History Forward Button
 				+SHorizontalBox::Slot()
 				.AutoWidth()
 				.Padding(4, 0)
 				[
-					SNew(SButton)
-					.ToolTipText( this, &SReferenceViewer::GetHistoryForwardTooltip )
-					.ButtonStyle(&FAppStyle::Get().GetWidgetStyle<FButtonStyle>("SimpleButton"))
-					.ForegroundColor(FSlateColor::UseStyle())
-					.OnClicked(this, &SReferenceViewer::ForwardClicked)
-					.IsEnabled(this, &SReferenceViewer::IsForwardEnabled)
-					[
-						SNew(SImage)
-						.ColorAndOpacity(FSlateColor::UseForeground())
-						.Image(FAppStyle::GetBrush("Icons.ArrowRight"))
-						.DesiredSizeOverride(FVector2D(20.0f, 20.0f))
-					]
-				]
-
-				+SHorizontalBox::Slot()
-				.AutoWidth()
-				.Padding(4, 0)
-				[
-					SNew(SSimpleComboButton)
-					.OnGetMenuContent(this, &SReferenceViewer::GetShowMenuContent)
-					.Icon(FAppStyle::GetBrush("Icons.Visibility"))
-					.HasDownArrow(true)
+					MakeToolBar()
 				]
 
 				// Path
 				+SHorizontalBox::Slot()
-				.Padding(4, 0)
+				.Padding(0, 0)
 				.FillWidth(1.f)
+				.VAlign(VAlign_Center)
 				[
 					SNew(SBorder)
 					.BorderImage( FAppStyle::GetBrush("ToolPanel.GroupBorder") )
@@ -236,14 +198,13 @@ void SReferenceViewer::Construct(const FArguments& InArgs)
 						.Style(FAppStyle::Get(), "ReferenceViewer.PathText")
 					]
 				]
-
-			
 			]
 		]
 
 		// Graph
 		+SVerticalBox::Slot()
 		.FillHeight(1.f)
+		.HAlign(HAlign_Fill)
 		[
 			SNew(SOverlay)
 
@@ -253,10 +214,15 @@ void SReferenceViewer::Construct(const FArguments& InArgs)
 			]
 
 			+SOverlay::Slot()
-			.HAlign(HAlign_Left)
 			.VAlign(VAlign_Top)
+			.HAlign(HAlign_Left)
 			.Padding(8)
 			[
+				SNew(SHorizontalBox)
+				+SHorizontalBox::Slot()
+				.AutoWidth()
+				[
+
 				SNew(SBorder)
 				.BorderImage( FAppStyle::GetBrush("ToolPanel.GroupBorder") )
 				[
@@ -404,6 +370,7 @@ void SReferenceViewer::Construct(const FArguments& InArgs)
 
 						+SHorizontalBox::Slot()
 						.VAlign(VAlign_Center)
+						.FillWidth(1.0)
 						.Padding(2.f)
 						[
 							SNew(STextBlock)
@@ -442,39 +409,21 @@ void SReferenceViewer::Construct(const FArguments& InArgs)
 						]
 					]
 
-					+ SVerticalBox::Slot()
-					.AutoHeight()
-					[
-						SNew(SHorizontalBox)
-						.Visibility_Lambda([]() { return (GShowToggleDeprecatedReferenceViewerLayout ? EVisibility::Visible : EVisibility::Collapsed); })
+				]
+				] // SHorizontalBox::Slot()
 
-						+ SHorizontalBox::Slot()
-						.VAlign(VAlign_Center)
-						.Padding(2.f)
-						[
-							SNew(STextBlock)
-							.Text(LOCTEXT("UseOldLayoutMechanism", "Use Deprecated Layout"))
-						]
+				+SHorizontalBox::Slot()
+				.AutoWidth()
+				.VAlign(VAlign_Top)
+				[
+					FilterCombo.ToSharedRef()
+				]
 
-						+SHorizontalBox::Slot()
-						.AutoWidth()
-						.VAlign(VAlign_Center)
-						.Padding(2.f)
-						[
-							SNew(SCheckBox)
-							.OnCheckStateChanged_Lambda( [this] (ECheckBoxState NewState) { 
-								if (GraphObj) 
-								{ 	
-									GraphObj->SetUseNodeInfos(NewState != ECheckBoxState::Checked); 
-									RebuildGraph();} 
-								} 
-							)
-							.IsChecked_Lambda([this] () 
-							{
-								return (GraphObj && GraphObj->GetUseNodeInfos()) ? ECheckBoxState::Unchecked: ECheckBoxState::Checked;
-							})
-						]
-					]
+				+SHorizontalBox::Slot()
+				.FillWidth(1.0)
+				.VAlign(VAlign_Top)
+				[
+					FilterWidget.ToSharedRef()
 				]
 			]
 
@@ -620,27 +569,21 @@ bool SReferenceViewer::IsForwardEnabled() const
 	return HistoryManager.CanGoForward();
 }
 
-FReply SReferenceViewer::BackClicked()
+void SReferenceViewer::BackClicked()
 {
 	HistoryManager.GoBack();
-
-	return FReply::Handled();
 }
 
-FReply SReferenceViewer::ForwardClicked()
+void SReferenceViewer::ForwardClicked()
 {
 	HistoryManager.GoForward();
-
-	return FReply::Handled();
 }
 
-FReply SReferenceViewer::RefreshClicked()
+void SReferenceViewer::RefreshClicked()
 {
 	RebuildGraph();
 	TriggerZoomToFit(0, 0);
 	RegisterActiveTimer(0.1f, FWidgetActiveTimerDelegate::CreateSP(this, &SReferenceViewer::TriggerZoomToFit));
-
-	return FReply::Handled();
 }
 
 void SReferenceViewer::GraphNavigateHistoryBack()
@@ -780,6 +723,38 @@ void SReferenceViewer::OnUpdateHistoryData(FReferenceViewerHistoryData& HistoryD
 	}
 }
 
+void SReferenceViewer::OnUpdateFilterBar()
+{
+	bRebuildingFilters = true; 
+
+	if (GraphObj)
+	{
+		const TSet<FName> AllClasses = GraphObj->GetAssetTypes();
+		if (Settings->AutoUpdateFilters())
+		{
+			FilterWidget->RemoveAllFilters();
+			for (FName AssetClass : AllClasses)
+			{
+				if (FilterWidget->DoesAssetTypeFilterExist(AssetClass))
+				{
+					FilterWidget->SetAssetTypeFilterCheckState(AssetClass, ECheckBoxState::Checked);
+				}
+			}
+
+			GraphObj->SetCurrentFilterCollection(FilterWidget->GetAllActiveFilters());
+		}
+
+		else
+		{
+			FString EmptyString;
+			FilterWidget->LoadSettings(EmptyString, EmptyString, EmptyString);
+			GraphObj->SetCurrentFilterCollection(FilterWidget->GetAllActiveFilters());
+		}
+	}
+
+	bRebuildingFilters = false;
+}
+
 void SReferenceViewer::OnSearchDepthEnabledChanged( ECheckBoxState NewState )
 {
 	Settings->SetSearchDepthLimitEnabled(NewState == ECheckBoxState::Checked);
@@ -816,7 +791,10 @@ void SReferenceViewer::OnSearchReferencerDepthCommitted(int32 NewValue)
 void SReferenceViewer::OnSearchBreadthEnabledChanged( ECheckBoxState NewState )
 {
 	Settings->SetSearchBreadthLimitEnabled(NewState == ECheckBoxState::Checked);
-	RebuildGraph();
+	if (GraphObj)
+	{
+		GraphObj->RefilterGraph();
+	}
 }
 
 ECheckBoxState SReferenceViewer::IsSearchBreadthEnabledChecked() const
@@ -1000,7 +978,10 @@ bool SReferenceViewer::IsCompactModeChecked() const
 void SReferenceViewer::OnShowDuplicatesChanged()
 {
 	Settings->SetShowDuplicatesEnabled(!Settings->IsShowDuplicates());
-	RebuildGraph();
+	if (GraphObj)
+	{
+		GraphObj->RefilterGraph();
+	}
 }
 
 bool SReferenceViewer::IsShowDuplicatesChecked() const
@@ -1011,7 +992,10 @@ bool SReferenceViewer::IsShowDuplicatesChecked() const
 void SReferenceViewer::OnShowEditorOnlyReferencesChanged()
 {
 	Settings->SetShowEditorOnlyReferencesEnabled(!Settings->IsShowEditorOnlyReferences());
-	RebuildGraph();
+	if (GraphObj)
+	{
+		GraphObj->RebuildGraph();
+	}
 }
 
 bool SReferenceViewer::IsShowEditorOnlyReferencesChecked() const
@@ -1069,7 +1053,10 @@ int32 SReferenceViewer::GetSearchBreadthCount() const
 void SReferenceViewer::OnSearchBreadthCommitted(int32 NewValue)
 {
 	Settings->SetSearchBreadthLimit(NewValue);
-	RebuildGraph();
+	if (GraphObj)
+	{
+		GraphObj->RefilterGraph();
+	}
 }
 
 void SReferenceViewer::RegisterActions()
@@ -1279,11 +1266,38 @@ void SReferenceViewer::RegisterActions()
 		FAssetManagerEditorCommands::Get().ShowCommentPath,
 		FExecuteAction::CreateLambda([this] { 
 			Settings->SetShowPathEnabled(!Settings->IsShowPath());
-			RebuildGraph();
+			if (GraphObj)
+			{
+				GraphObj->RefilterGraph();
+			}
 		}),
 		FCanExecuteAction(),	
 		FIsActionChecked::CreateLambda([this] {return Settings->IsShowPath();}));
 
+	ReferenceViewerActions->MapAction(
+		FAssetManagerEditorCommands::Get().Filters,
+		FExecuteAction::CreateLambda([this] { 
+			Settings->SetFiltersEnabled(!Settings->GetFiltersEnabled());
+			if (GraphObj)
+			{
+				GraphObj->RefilterGraph();
+			}
+		}),
+		FCanExecuteAction(),	
+		FIsActionChecked::CreateLambda([this] {return Settings->GetFiltersEnabled();}));
+
+	ReferenceViewerActions->MapAction(
+		FAssetManagerEditorCommands::Get().AutoFilters,
+		FExecuteAction::CreateLambda([this] { 
+			Settings->SetAutoUpdateFilters(!Settings->AutoUpdateFilters());
+			if (GraphObj)
+			{
+				OnUpdateFilterBar();
+				GraphObj->RefilterGraph();
+			}
+		}),
+		FCanExecuteAction::CreateLambda([this] {return Settings->GetFiltersEnabled();}),
+		FIsActionChecked::CreateLambda([this] {return Settings->AutoUpdateFilters();}));
 
 	ReferenceViewerActions->MapAction(
 		FAssetManagerEditorCommands::Get().CopyPaths,
@@ -1865,4 +1879,75 @@ TSharedRef<SWidget> SReferenceViewer::GetShowMenuContent()
 	return MenuBuilder.MakeWidget();
 }
 
+TSharedRef<SWidget> SReferenceViewer::MakeToolBar()
+{
+
+	FToolBarBuilder ToolBarBuilder(ReferenceViewerActions, FMultiBoxCustomization::None, TSharedPtr<FExtender>(), true);
+	ToolBarBuilder.SetStyle(&FReferenceViewerStyle::Get(), "AssetEditorToolbar");
+	ToolBarBuilder.BeginSection("Test");
+
+	ToolBarBuilder.AddToolBarButton(
+		FUIAction(FExecuteAction::CreateSP(this, &SReferenceViewer::RefreshClicked)),
+		NAME_None,
+		TAttribute<FText>(),
+		TAttribute<FText>(),
+		FSlateIcon(FReferenceViewerStyle::Get().GetStyleSetName(), "Icons.Refresh"));
+
+	ToolBarBuilder.AddToolBarButton(
+		FUIAction(
+			FExecuteAction::CreateSP(this, &SReferenceViewer::BackClicked),
+			FCanExecuteAction::CreateSP(this, &SReferenceViewer::IsBackEnabled)
+		),
+		NAME_None,
+		TAttribute<FText>(),
+		TAttribute<FText>::CreateSP(this, &SReferenceViewer::GetHistoryBackTooltip),
+		FSlateIcon(FReferenceViewerStyle::Get().GetStyleSetName(), "Icons.ArrowLeft"));
+
+
+	ToolBarBuilder.AddToolBarButton(
+		FUIAction(
+			FExecuteAction::CreateSP(this, &SReferenceViewer::ForwardClicked),
+			FCanExecuteAction::CreateSP(this, &SReferenceViewer::IsForwardEnabled)
+		),
+		NAME_None,
+		TAttribute<FText>(),
+		TAttribute<FText>::CreateSP(this, &SReferenceViewer::GetHistoryForwardTooltip),
+		FSlateIcon(FReferenceViewerStyle::Get().GetStyleSetName(), "Icons.ArrowRight"));
+
+	ToolBarBuilder.AddSeparator();
+
+	ToolBarBuilder.AddComboButton( 
+		FUIAction(),
+		FOnGetContent::CreateSP(this, &SReferenceViewer::GetShowMenuContent),
+		TAttribute<FText>(),
+		TAttribute<FText>(),
+		FSlateIcon(FReferenceViewerStyle::Get().GetStyleSetName(), "Icons.Visibility"),
+		/*bInSimpleComboBox*/ false);
+
+	ToolBarBuilder.AddToolBarButton(FAssetManagerEditorCommands::Get().ShowDuplicates,
+		NAME_None,
+		TAttribute<FText>(),
+		TAttribute<FText>(),
+		FSlateIcon(FReferenceViewerStyle::Get().GetStyleSetName(), "Icons.Duplicate"));
+
+	ToolBarBuilder.AddSeparator();
+
+	ToolBarBuilder.AddToolBarButton(FAssetManagerEditorCommands::Get().Filters,
+		NAME_None,
+		TAttribute<FText>(),
+		TAttribute<FText>(),
+		FSlateIcon(FReferenceViewerStyle::Get().GetStyleSetName(), "Icons.Filters"));
+
+	ToolBarBuilder.AddToolBarButton(FAssetManagerEditorCommands::Get().AutoFilters,
+		NAME_None,
+		TAttribute<FText>(),
+		TAttribute<FText>(),
+		FSlateIcon(FReferenceViewerStyle::Get().GetStyleSetName(), "Icons.AutoFilters"));
+
+	ToolBarBuilder.EndSection();
+
+
+	return ToolBarBuilder.MakeWidget();
+}
+
 #undef LOCTEXT_NAMESPACE
\ No newline at end of file
diff --git a/Engine/Plugins/Editor/AssetManagerEditor/Source/AssetManagerEditor/Private/ReferenceViewer/SReferenceViewer.h b/Engine/Plugins/Editor/AssetManagerEditor/Source/AssetManagerEditor/Private/ReferenceViewer/SReferenceViewer.h
--- a/Engine/Plugins/Editor/AssetManagerEditor/Source/AssetManagerEditor/Private/ReferenceViewer/SReferenceViewer.h
+++ b/Engine/Plugins/Editor/AssetManagerEditor/Source/AssetManagerEditor/Private/ReferenceViewer/SReferenceViewer.h
@@ -15,9 +15,10 @@
 #include "AssetManagerEditorModule.h"
 #include "Containers/ArrayView.h"
 #include "ReferenceViewer/ReferenceViewerSettings.h"
+#include "ReferenceViewer/EdGraph_ReferenceViewer.h"
+#include "ReferenceViewer/SReferenceViewerFilterBar.h"
 
 class UEdGraph;
-class UEdGraph_ReferenceViewer;
 
 /**
  * 
@@ -70,13 +71,13 @@ private:
 	bool IsForwardEnabled() const;
 
 	/** Handler for clicking the history back button */
-	FReply BackClicked();
+	void BackClicked();
 
 	/** Handler for clicking the history forward button */
-	FReply ForwardClicked();
+	void ForwardClicked();
 
 	/** Refresh the current view */
-	FReply RefreshClicked();
+	void RefreshClicked();
 
 	/** Handler for when the graph panel tells us to go back in history (like using the mouse thumb button) */
 	void GraphNavigateHistoryBack();
@@ -105,6 +106,8 @@ private:
 	void OnApplyHistoryData(const FReferenceViewerHistoryData& History);
 
 	void OnUpdateHistoryData(FReferenceViewerHistoryData& HistoryData) const;
+
+	void OnUpdateFilterBar();
 	
 	void OnSearchDepthEnabledChanged( ECheckBoxState NewState );
 	ECheckBoxState IsSearchDepthEnabledChecked() const;
@@ -192,6 +195,9 @@ private:
 	EActiveTimerReturnType TriggerZoomToFit(double InCurrentTime, float InDeltaTime);
 private:
 
+	TSharedRef<SWidget> MakeToolBar();
+
+
 	/** The manager that keeps track of history data for this browser */
 	FReferenceViewerHistoryManager HistoryManager;
 
@@ -203,6 +209,8 @@ private:
 	TSharedPtr<SWidget> DependencyCountBox;
 	TSharedPtr<SWidget> BreadthLimitBox;
 
+	TSharedPtr< SReferenceViewerFilterBar > FilterWidget;
+
 	UEdGraph_ReferenceViewer* GraphObj;
 
 	UReferenceViewerSettings* Settings;
@@ -243,6 +251,9 @@ private:
 	/** Whether to visually show to the user the option of "Compact Mode" */
 	bool bShowCompactMode;
 
+	/** A recursion check so as to avoid the rebuild of the graph if we are currently rebuilding the filters */
+	bool bRebuildingFilters;
+
 	/** Handle to know if dirty */
 	FDelegateHandle AssetRefreshHandle;
 };
diff --git a/Engine/Plugins/Editor/AssetManagerEditor/Source/AssetManagerEditor/Private/ReferenceViewer/SReferenceViewerFilterBar.cpp b/Engine/Plugins/Editor/AssetManagerEditor/Source/AssetManagerEditor/Private/ReferenceViewer/SReferenceViewerFilterBar.cpp
new file mode 100644
--- /dev/null
+++ b/Engine/Plugins/Editor/AssetManagerEditor/Source/AssetManagerEditor/Private/ReferenceViewer/SReferenceViewerFilterBar.cpp
@@ -0,0 +1,47 @@
+// Copyright Epic Games, Inc. All Rights Reserved.
+
+#pragma once
+
+#include "ReferenceViewer/SReferenceViewerFilterBar.h"
+
+#include "CoreMinimal.h"
+
+void SReferenceViewerFilterBar::SaveSettings(const FString& IniFilename, const FString& IniSection, const FString& SettingsString) const
+{
+	TArray<FilterState> SaveFilters;
+	if (UReferenceViewerSettings* Settings = GetMutableDefault<UReferenceViewerSettings>())
+	{
+		// Only save & load user filters, not autofilters
+		if (!Settings->AutoUpdateFilters())
+		{
+			for (TSharedRef<SAssetFilter> CurrentAssetFilter : AssetFilters)
+			{
+				FName AssetFilterName = FName(CurrentAssetFilter->GetCustomClassFilterData()->GetFilterName());
+				SaveFilters.Add(FilterState(AssetFilterName, CurrentAssetFilter->IsEnabled()));
+			}
+			Settings->SetUserFilters(SaveFilters);
+		}
+	}
+}
+
+void SReferenceViewerFilterBar::LoadSettings(const FString& IniFilename, const FString& IniSection, const FString& SettingsString) 
+{
+	if (UReferenceViewerSettings* Settings = GetMutableDefault<UReferenceViewerSettings>())
+	{
+		// Only save & load user filters, not autofilters
+		if (!Settings->AutoUpdateFilters())
+		{
+			TArray<FilterState> SavedFilters = Settings->GetUserFilters();
+			RemoveAllFilters();
+			for (FilterState& State : SavedFilters)
+			{
+				if (DoesAssetTypeFilterExist(State.FilterName))
+				{
+					SetAssetTypeFilterCheckState(State.FilterName, ECheckBoxState::Checked);
+					ToggleAssetTypeFilterEnabled(State.FilterName, State.bIsEnabled);
+				}
+			}
+		}
+	}
+}
+
diff --git a/Engine/Plugins/Editor/AssetManagerEditor/Source/AssetManagerEditor/Private/ReferenceViewer/SReferenceViewerFilterBar.h b/Engine/Plugins/Editor/AssetManagerEditor/Source/AssetManagerEditor/Private/ReferenceViewer/SReferenceViewerFilterBar.h
new file mode 100644
--- /dev/null
+++ b/Engine/Plugins/Editor/AssetManagerEditor/Source/AssetManagerEditor/Private/ReferenceViewer/SReferenceViewerFilterBar.h
@@ -0,0 +1,21 @@
+// Copyright Epic Games, Inc. All Rights Reserved.
+
+#pragma once
+
+#include "CoreMinimal.h"
+#include "Filters/SFilterBar.h"
+#include "ReferenceViewer/ReferenceViewerSettings.h"
+#include "ReferenceViewer/EdGraph_ReferenceViewer.h"
+
+class SReferenceViewerFilterBar : public SFilterBar< FReferenceNodeInfo& > 
+{
+
+public:
+
+	/** Saves any settings to config that should be persistent between editor sessions */
+	void SaveSettings(const FString& IniFilename, const FString& IniSection, const FString& SettingsString) const override;
+
+	/** Loads any settings to config that should be persistent between editor sessions */
+	void LoadSettings(const FString& IniFilename, const FString& IniSection, const FString& SettingsString) override;
+
+};
\ No newline at end of file
diff --git a/Engine/Plugins/Editor/AssetManagerEditor/Source/AssetManagerEditor/Public/ReferenceViewer/EdGraphNode_Reference.h b/Engine/Plugins/Editor/AssetManagerEditor/Source/AssetManagerEditor/Public/ReferenceViewer/EdGraphNode_Reference.h
--- a/Engine/Plugins/Editor/AssetManagerEditor/Source/AssetManagerEditor/Public/ReferenceViewer/EdGraphNode_Reference.h
+++ b/Engine/Plugins/Editor/AssetManagerEditor/Source/AssetManagerEditor/Public/ReferenceViewer/EdGraphNode_Reference.h
@@ -44,6 +44,13 @@ class ASSETMANAGEREDITOR_API UEdGraphNode_Reference : public UEdGraphNode
 	bool IsPackage() const;
 	bool IsCollapsed() const;
 	bool IsADuplicate() const { return bIsADuplicate; }
+
+	// Nodes that are filtered out still may still show because they
+	// are between nodes that pass the filter and the root.  This "filtered"
+	// bool allows us to render these in-between nodes differently
+	void SetIsFiltered(bool bInFiltered);
+	bool GetIsFiltered() const;
+
 	FAssetData GetAssetData() const;
 
 	UEdGraphPin* GetDependencyPin();
@@ -64,6 +71,7 @@ private:
 	bool bIsPrimaryAsset;
 	bool bIsCollapsed;
 	bool bIsADuplicate;
+	bool bIsFiltered;
 
 	FAssetData CachedAssetData;
 	FLinearColor AssetTypeColor;
diff --git a/Engine/Plugins/Editor/AssetManagerEditor/Source/AssetManagerEditor/Public/ReferenceViewer/EdGraph_ReferenceViewer.h b/Engine/Plugins/Editor/AssetManagerEditor/Source/AssetManagerEditor/Public/ReferenceViewer/EdGraph_ReferenceViewer.h
--- a/Engine/Plugins/Editor/AssetManagerEditor/Source/AssetManagerEditor/Public/ReferenceViewer/EdGraph_ReferenceViewer.h
+++ b/Engine/Plugins/Editor/AssetManagerEditor/Source/AssetManagerEditor/Public/ReferenceViewer/EdGraph_ReferenceViewer.h
@@ -8,6 +8,7 @@
 #include "AssetManagerEditorModule.h"
 #include "EdGraph/EdGraph.h"
 #include "Misc/AssetRegistryInterface.h"
+#include "Misc/FilterCollection.h"
 #include "EdGraph_ReferenceViewer.generated.h"
 
 class FAssetThumbnailPool;
@@ -19,7 +20,7 @@ enum class EDependencyPinCategory;
 /*
 *  Holds asset information for building reference graph
 */ 
-struct ReferenceNodeInfo
+struct FReferenceNodeInfo
 {
 	FAssetIdentifier AssetId;
 
@@ -36,7 +37,7 @@ struct ReferenceNodeInfo
 
 	int32 OverflowCount;
 
-	ReferenceNodeInfo(const FAssetIdentifier& InAssetId, bool InbReferencers);
+	FReferenceNodeInfo(const FAssetIdentifier& InAssetId, bool InbReferencers);
 
 	bool IsFirstParent(const FAssetIdentifier& InParentId) const;
 
@@ -50,6 +51,9 @@ struct ReferenceNodeInfo
 	// how many nodes worth of children require vertical spacing 
 	int32 ChildProvisionSize;
 
+	// Whether or not this nodeinfo passed the current filters 
+	bool PassedFilters;
+
 };
 
 
@@ -78,25 +82,29 @@ public:
 	/** Force the graph to rebuild */
 	class UEdGraphNode_Reference* RebuildGraph();
 
+	/** Refilters the nodes, more efficient that a full rebuild.  This function is preferred when the assets, reference types or depth hasn't changed, meaning the NodeInfos didn't change, just 
+	 * the presentation or filtering */
+	class UEdGraphNode_Reference* RefilterGraph();
+
 	using FIsPackageNamePassingFilterCallback = TFunction<bool(FName)>;
 	void SetIsPackageNamePassingFilterCallback(const TOptional<FIsPackageNamePassingFilterCallback>& InIsPackageNamePassingFilterCallback) { IsPackageNamePassingFilterCallback = InIsPackageNamePassingFilterCallback; }
 
 	FName GetCurrentCollectionFilter() const;
 	void SetCurrentCollectionFilter(FName NewFilter);
 
-	/* Temporary variable that allows reverting to deprecated layout methods */
-	bool GetUseNodeInfos() const { return bUseNodeInfos; }
-	void SetUseNodeInfos(bool InbUseNodeInfos) { bUseNodeInfos= InbUseNodeInfos; }
+	/* Delegate type to notify when the assets or NodeInfos have changed as opposed to when the filters changed */
+	FSimpleMulticastDelegate& OnAssetsChanged() { return OnAssetsChangedDelegate; }
+
+	/* Not to be confused with the above Content Browser Collection name, this is a TFiltercollection, a list of active filters */
+	void SetCurrentFilterCollection(TSharedPtr< TFilterCollection<FReferenceNodeInfo&> > NewFilterCollection);
+
+	/* Returns a set of unique asset types as UClass* */
+	const TSet<FName>& GetAssetTypes() const { return CurrentClasses; }
 
 private:
 	void SetReferenceViewer(TSharedPtr<SReferenceViewer> InViewer);
 	UEdGraphNode_Reference* ConstructNodes(const TArray<FAssetIdentifier>& GraphRootIdentifiers, const FIntPoint& GraphRootOrigin);
-	// RecursivelyGatherSizes is now deprecated, use RecursivelyPopulateNodeInfos with RecursivelyCreateNodes
-	int32 RecursivelyGatherSizes(bool bReferencers, const TArray<FAssetIdentifier>& Identifiers, const TSet<FName>& AllowedPackageNames, int32 CurrentDepth, int32 MaxDepth, TSet<FAssetIdentifier>& VisitedNames, TMap<FAssetIdentifier, int32>& OutNodeSizes) const;
-	void GatherAssetData(const TSet<FName>& AllPackageNames, TMap<FName, FAssetData>& OutPackageToAssetDataMap) const;
 
-	// RecursivelyConstructNodes is now deprecated, use RecursivelyPopulateNodeInfos with RecursivelyCreateNodes
-	class UEdGraphNode_Reference* RecursivelyConstructNodes(bool bReferencers, UEdGraphNode_Reference* RootNode, const TArray<FAssetIdentifier>& Identifiers, const FIntPoint& NodeLoc, const TMap<FAssetIdentifier, int32>& NodeSizes, const TMap<FName, FAssetData>& PackagesToAssetDataMap, const TSet<FName>& AllowedPackageNames, int32 CurrentDepth, int32 MaxDepth, TSet<FAssetIdentifier>& VisitedNames);
 
 	bool ExceedsMaxSearchDepth(int32 Depth, int32 MaxDepth) const;
 	bool ExceedsMaxSearchBreadth(int32 Breadth) const;
@@ -105,7 +113,13 @@ private:
 	UEdGraphNode_Reference* CreateReferenceNode();
 
 	/* Generates a NodeInfo structure then used to generate and layout the graph nodes */
-	void RecursivelyPopulateNodeInfos(bool bReferencers, const FAssetIdentifier& AssetId, TMap<FAssetIdentifier, ReferenceNodeInfo>& NodeInfos, int32 CurrentDepth, int32 MaxDepth);
+	void RecursivelyPopulateNodeInfos(bool bReferencers, const FAssetIdentifier& AssetId, TMap<FAssetIdentifier, FReferenceNodeInfo>& NodeInfos, int32 CurrentDepth, int32 MaxDepth);
+
+	/* Marks up the NodeInfos with updated filter information and provision sizes */
+	void RecursivelyFilterNodeInfos(const FAssetIdentifier& InAssetId, TMap<FAssetIdentifier, FReferenceNodeInfo>& NodeInfos, int32 CurrentDepth, int32 MaxDepth);
+
+	/* Searches for the AssetData for the list of packages derived from the AssetReferences  */
+	void GatherAssetData(const TSet<FName>& AllPackageNames, TMap<FName, FAssetData>& OutPackageToAssetDataMap) const;
 
 	/* Uses the NodeInfos map to generate and layout the graph nodes */
 	UEdGraphNode_Reference* RecursivelyCreateNodes(
@@ -114,7 +128,7 @@ private:
 		const FIntPoint& InNodeLoc, 
 		const FAssetIdentifier& InParentId, 
 		UEdGraphNode_Reference* InParentNode, 
-		TMap<FAssetIdentifier, ReferenceNodeInfo>& InNodeInfos, 
+		TMap<FAssetIdentifier, FReferenceNodeInfo>& InNodeInfos, 
 		int32 InCurrentDepth, 
 		int32 InMaxDepth, 
 		bool bIsRoot = false
@@ -144,17 +158,28 @@ private:
 	/** Current collection filter. NAME_None for no filter */
 	FName CurrentCollectionFilter;
 
+	/** A set of the unique class types referenced */
+	TSet<FName> CurrentClasses;
 
 	/* This is a convenience toggle to switch between the old & new methods for computing & displaying the graph */
 	bool bUseNodeInfos;
 
+	/* Cached Reference Information used to quickly refilter */
+	TMap<FAssetIdentifier, FReferenceNodeInfo> ReferencerNodeInfos;
+	TMap<FAssetIdentifier, FReferenceNodeInfo> DependencyNodeInfos;
+
 	TOptional<FIsPackageNamePassingFilterCallback> IsPackageNamePassingFilterCallback;
 
 	/** List of packages the current collection filter allows */
 	TSet<FName> CurrentCollectionPackages;
 
+	/** Current filter collection */
+	TSharedPtr< TFilterCollection<FReferenceNodeInfo & > > FilterCollection;
 
 	UReferenceViewerSettings* Settings;
 
+	/* A delegate to notify when the underlying assets changed (usually through a root or depth change) */
+	FSimpleMulticastDelegate OnAssetsChangedDelegate;
+
 	friend SReferenceViewer;
 };
diff --git a/Engine/Plugins/Editor/AssetManagerEditor/Source/AssetManagerEditor/Public/ReferenceViewer/ReferenceViewerSettings.h b/Engine/Plugins/Editor/AssetManagerEditor/Source/AssetManagerEditor/Public/ReferenceViewer/ReferenceViewerSettings.h
--- a/Engine/Plugins/Editor/AssetManagerEditor/Source/AssetManagerEditor/Public/ReferenceViewer/ReferenceViewerSettings.h
+++ b/Engine/Plugins/Editor/AssetManagerEditor/Source/AssetManagerEditor/Public/ReferenceViewer/ReferenceViewerSettings.h
@@ -10,6 +10,27 @@
 /**
  *  Project based Reference Viewer Saved Settings 
  */
+USTRUCT()
+struct FilterState
+{
+	GENERATED_BODY();
+
+	FilterState()
+	: FilterName(NAME_None)
+	, bIsEnabled(false)
+	{}
+
+	FilterState(FName InFilterName, bool InState)
+	: FilterName(InFilterName)
+	, bIsEnabled(InState)
+	{}
+
+	UPROPERTY()
+	FName FilterName;
+
+	UPROPERTY()
+	bool bIsEnabled;
+};
 
 UCLASS(config=EditorPerProjectUserSettings)
 class UReferenceViewerSettings : public UObject
@@ -71,6 +92,16 @@ public:
 	bool IsShowPath() const;
 	void SetShowPathEnabled(bool newEnabled);
 
+	bool GetFiltersEnabled() const;
+	void SetFiltersEnabled(bool newEnabled);
+
+	bool AutoUpdateFilters() const;
+	void SetAutoUpdateFilters(bool newEnabled);
+
+	const TArray<FilterState>& GetUserFilters() const;
+	void SetUserFilters(TArray<FilterState>& InFilters);
+
+
 private:
 	/* Whether to limit the search depth for Referencers & Dependencies */
 	UPROPERTY(config)
@@ -144,4 +175,15 @@ private:
 	UPROPERTY(config)
 	bool bIsShowPath;
 
+	/* This turns on/off any filtering done though the SFilterBar */
+	UPROPERTY(config)
+	bool bFiltersEnabled;
+
+	/* When true, the filters bar auto updates based on the node types, otherwise user filters will be used */
+	UPROPERTY(config)
+	bool bAutoUpdateFilters;
+
+	/* The list of filters the user has built up */
+	UPROPERTY(config)
+	TArray<FilterState> UserFilters;
 };
diff --git a/Engine/Source/Runtime/SlateCore/Private/Styling/StarshipCoreStyle.cpp b/Engine/Source/Runtime/SlateCore/Private/Styling/StarshipCoreStyle.cpp
--- a/Engine/Source/Runtime/SlateCore/Private/Styling/StarshipCoreStyle.cpp
+++ b/Engine/Source/Runtime/SlateCore/Private/Styling/StarshipCoreStyle.cpp
@@ -340,6 +340,7 @@ TSharedRef<ISlateStyle> FStarshipCoreStyle::Create()
 
 		Style->Set("Icons.Import", new IMAGE_BRUSH_SVG("Starship/Common/import", Icon16x16));
 		Style->Set("Icons.Filter", new IMAGE_BRUSH_SVG("Starship/Common/filter", Icon16x16));
+		Style->Set("Icons.AutoFilter", new IMAGE_BRUSH_SVG("Starship/Common/FilterAuto", Icon16x16));
 
 		Style->Set("Icons.Lock", new IMAGE_BRUSH_SVG("Starship/Common/lock", Icon16x16));
 		Style->Set("Icons.Unlock", new IMAGE_BRUSH_SVG("Starship/Common/lock-unlocked", Icon16x16));
