#UE Fix compile error fallout from ANY_PACKAGE change

diff --git a/Engine/Plugins/Editor/AssetManagerEditor/Source/AssetManagerEditor/Private/ReferenceViewer/EdGraph_ReferenceViewer.cpp b/Engine/Plugins/Editor/AssetManagerEditor/Source/AssetManagerEditor/Private/ReferenceViewer/EdGraph_ReferenceViewer.cpp
--- a/Engine/Plugins/Editor/AssetManagerEditor/Source/AssetManagerEditor/Private/ReferenceViewer/EdGraph_ReferenceViewer.cpp
+++ b/Engine/Plugins/Editor/AssetManagerEditor/Source/AssetManagerEditor/Private/ReferenceViewer/EdGraph_ReferenceViewer.cpp
@@ -233,13 +233,13 @@ UEdGraphNode_Reference* UEdGraph_ReferenceViewer::ConstructNodes(const TArray<FA
 		GatherAssetData(AllPackageNames, PackagesToAssetDataMap);
 
 		// Store the AssetData in the NodeInfos and collect Asset Type UClasses to populate the filters
-		TSet<FName> AllClasses;
+		TSet<FTopLevelAssetPath> AllClasses;
 		for (TPair<FAssetIdentifier, FReferenceNodeInfo>&  InfoPair : NewReferenceNodeInfos)
 		{
 			InfoPair.Value.AssetData = PackagesToAssetDataMap.FindRef(InfoPair.Key.PackageName);
 			if (InfoPair.Value.AssetData.IsValid())
 			{
-				AllClasses.Add(InfoPair.Value.AssetData.AssetClass);
+				AllClasses.Add(InfoPair.Value.AssetData.AssetClassPath);
 			}
 		}
 
@@ -248,7 +248,7 @@ UEdGraphNode_Reference* UEdGraph_ReferenceViewer::ConstructNodes(const TArray<FA
 			InfoPair.Value.AssetData = PackagesToAssetDataMap.FindRef(InfoPair.Key.PackageName);
 			if (InfoPair.Value.AssetData.IsValid())
 			{
-				AllClasses.Add(InfoPair.Value.AssetData.AssetClass);
+				AllClasses.Add(InfoPair.Value.AssetData.AssetClassPath);
 			}
 		}
 
diff --git a/Engine/Plugins/Editor/AssetManagerEditor/Source/AssetManagerEditor/Private/ReferenceViewer/SReferenceViewer.cpp b/Engine/Plugins/Editor/AssetManagerEditor/Source/AssetManagerEditor/Private/ReferenceViewer/SReferenceViewer.cpp
--- a/Engine/Plugins/Editor/AssetManagerEditor/Source/AssetManagerEditor/Private/ReferenceViewer/SReferenceViewer.cpp
+++ b/Engine/Plugins/Editor/AssetManagerEditor/Source/AssetManagerEditor/Private/ReferenceViewer/SReferenceViewer.cpp
@@ -729,12 +729,14 @@ void SReferenceViewer::OnUpdateFilterBar()
 
 	if (GraphObj)
 	{
-		const TSet<FName> AllClasses = GraphObj->GetAssetTypes();
+		const TSet<FTopLevelAssetPath> AllClasses = GraphObj->GetAssetTypes();
 		if (Settings->AutoUpdateFilters())
 		{
 			FilterWidget->RemoveAllFilters();
-			for (FName AssetClass : AllClasses)
+			for (const FTopLevelAssetPath& AssetClassPath : AllClasses)
 			{
+				// @todo make DoesAssetTypeFilterExist/SetAssetTypeFilterCheckState use FTopLevelAssetPath
+				const FName AssetClass = FName(*AssetClassPath.ToString());
 				if (FilterWidget->DoesAssetTypeFilterExist(AssetClass))
 				{
 					FilterWidget->SetAssetTypeFilterCheckState(AssetClass, ECheckBoxState::Checked);
diff --git a/Engine/Plugins/Editor/AssetManagerEditor/Source/AssetManagerEditor/Public/ReferenceViewer/EdGraph_ReferenceViewer.h b/Engine/Plugins/Editor/AssetManagerEditor/Source/AssetManagerEditor/Public/ReferenceViewer/EdGraph_ReferenceViewer.h
--- a/Engine/Plugins/Editor/AssetManagerEditor/Source/AssetManagerEditor/Public/ReferenceViewer/EdGraph_ReferenceViewer.h
+++ b/Engine/Plugins/Editor/AssetManagerEditor/Source/AssetManagerEditor/Public/ReferenceViewer/EdGraph_ReferenceViewer.h
@@ -99,7 +99,7 @@ public:
 	void SetCurrentFilterCollection(TSharedPtr< TFilterCollection<FReferenceNodeInfo&> > NewFilterCollection);
 
 	/* Returns a set of unique asset types as UClass* */
-	const TSet<FName>& GetAssetTypes() const { return CurrentClasses; }
+	const TSet<FTopLevelAssetPath>& GetAssetTypes() const { return CurrentClasses; }
 
 private:
 	void SetReferenceViewer(TSharedPtr<SReferenceViewer> InViewer);
@@ -159,7 +159,7 @@ private:
 	FName CurrentCollectionFilter;
 
 	/** A set of the unique class types referenced */
-	TSet<FName> CurrentClasses;
+	TSet<FTopLevelAssetPath> CurrentClasses;
 
 	/* This is a convenience toggle to switch between the old & new methods for computing & displaying the graph */
 	bool bUseNodeInfos;
diff --git a/Engine/Source/Editor/EditorWidgets/Public/Filters/AssetFilter.h b/Engine/Source/Editor/EditorWidgets/Public/Filters/AssetFilter.h
--- a/Engine/Source/Editor/EditorWidgets/Public/Filters/AssetFilter.h
+++ b/Engine/Source/Editor/EditorWidgets/Public/Filters/AssetFilter.h
@@ -18,7 +18,7 @@ class FAssetFilter : public FFilterBase<FilterType>
 public:
 
 	/* A delegate that the user can use to specify how to compare their Filter Item to a list of Asset Types represented as FNames */
-	DECLARE_DELEGATE_RetVal_TwoParams(bool, FCompareItemWithClassNames, FilterType /* InItem */, const TSet<FName>& /* InClassNames */);
+	DECLARE_DELEGATE_RetVal_TwoParams(bool, FCompareItemWithClassNames, FilterType /* InItem */, const TSet<FTopLevelAssetPath>& /* InClassNames */);
 
 	/* A delegate that the user can use to specify how to convert their Filter Item into an FAssetData
 	 * @return: true if the Filter Item was successfully converted, false if not (the current item will not be tested for asset filters if so)
@@ -110,7 +110,7 @@ public:
 		// If not, try to use the comparison function
 		else if(OnCompareItem.IsBound())
 		{
-			return OnCompareItem.Execute(InItem, CompiledBackendFilter.ClassNames);
+			return OnCompareItem.Execute(InItem, CompiledBackendFilter.ClassPaths);
 		}
 
 		return false;
