Merging CL#20448281

Merge fix

diff --git a/Engine/Plugins/Editor/AssetManagerEditor/Source/AssetManagerEditor/Private/ReferenceViewer/EdGraphNode_Reference.cpp b/Engine/Plugins/Editor/AssetManagerEditor/Source/AssetManagerEditor/Private/ReferenceViewer/EdGraphNode_Reference.cpp
--- a/Engine/Plugins/Editor/AssetManagerEditor/Source/AssetManagerEditor/Private/ReferenceViewer/EdGraphNode_Reference.cpp
+++ b/Engine/Plugins/Editor/AssetManagerEditor/Source/AssetManagerEditor/Private/ReferenceViewer/EdGraphNode_Reference.cpp
@@ -292,4 +292,14 @@ bool UEdGraphNode_Reference::IsCollapsed() const
 	return bIsCollapsed;
 }
 
+void UEdGraphNode_Reference::SetIsFiltered(bool bInFiltered)
+{
+	bIsFiltered = bInFiltered;	
+}
+
+bool UEdGraphNode_Reference::GetIsFiltered() const
+{
+	return bIsFiltered;
+}
+
 #undef LOCTEXT_NAMESPACE
