Reference Viewer fix up Icon & Class Name

diff --git a/Engine/Plugins/Editor/AssetManagerEditor/Source/AssetManagerEditor/Private/ReferenceViewer/EdGraphNode_Reference.cpp b/Engine/Plugins/Editor/AssetManagerEditor/Source/AssetManagerEditor/Private/ReferenceViewer/EdGraphNode_Reference.cpp
--- a/Engine/Plugins/Editor/AssetManagerEditor/Source/AssetManagerEditor/Private/ReferenceViewer/EdGraphNode_Reference.cpp
+++ b/Engine/Plugins/Editor/AssetManagerEditor/Source/AssetManagerEditor/Private/ReferenceViewer/EdGraphNode_Reference.cpp
@@ -38,7 +38,7 @@ void UEdGraphNode_Reference::SetupReferenceNode(const FIntPoint& NodeLoc, const
 	Identifiers = NewIdentifiers;
 	const FAssetIdentifier& First = NewIdentifiers[0];
 	FString MainAssetName = InAssetData.AssetName.ToString();
-	FString AssetTypeName = InAssetData.AssetClassPath.ToString();
+	FString AssetTypeName = InAssetData.AssetClassPath.GetAssetName().ToString();
 
 	FAssetToolsModule& AssetToolsModule = FModuleManager::LoadModuleChecked<FAssetToolsModule>(TEXT("AssetTools"));	
 	if (UClass* AssetClass = InAssetData.GetClass())
@@ -49,7 +49,7 @@ void UEdGraphNode_Reference::SetupReferenceNode(const FIntPoint& NodeLoc, const
 			AssetTypeColor = AssetTypeActions.Pin()->GetTypeColor();
 		}
 	}
-	AssetBrush = FSlateIcon("EditorStyle", FName( *("ClassIcon." + InAssetData.AssetClassPath.ToString())));
+	AssetBrush = FSlateIcon("EditorStyle", FName( *("ClassIcon." + AssetTypeName)));
 
 	bIsCollapsed = false;
 	bIsPackage = true;
