Reference Viewer - Cleanup Specialty Nodes display from ANY_PACKAGE fallout

diff --git a/Engine/Plugins/Editor/AssetManagerEditor/Source/AssetManagerEditor/Private/ReferenceViewer/EdGraphNode_Reference.cpp b/Engine/Plugins/Editor/AssetManagerEditor/Source/AssetManagerEditor/Private/ReferenceViewer/EdGraphNode_Reference.cpp
--- a/Engine/Plugins/Editor/AssetManagerEditor/Source/AssetManagerEditor/Private/ReferenceViewer/EdGraphNode_Reference.cpp
+++ b/Engine/Plugins/Editor/AssetManagerEditor/Source/AssetManagerEditor/Private/ReferenceViewer/EdGraphNode_Reference.cpp
@@ -55,21 +55,30 @@ void UEdGraphNode_Reference::SetupReferenceNode(const FIntPoint& NodeLoc, const
 	bIsPackage = true;
 	bAllowThumbnail = bInAllowThumbnail;
 	bIsADuplicate = bInIsADuplicate;
-	
+
 	FPrimaryAssetId PrimaryAssetID = NewIdentifiers[0].GetPrimaryAssetId();
-	if (PrimaryAssetID.IsValid())
+	if (PrimaryAssetID.IsValid())  // Management References (PrimaryAssetIDs)
 	{
 		MainAssetName = PrimaryAssetID.PrimaryAssetName.ToString();
 		AssetTypeName = PrimaryAssetID.PrimaryAssetType.ToString();
 		bIsPackage = false;
 		bIsPrimaryAsset = true;
 	}
-	else if (NewIdentifiers[0].IsValue())
+	else if (First.IsValue()) // Searchable Names (GamePlay Tags)
 	{
 		MainAssetName = First.ValueName.ToString();
 		AssetTypeName = First.ObjectName.ToString();
 		bIsPackage = false;
 	}
+	else if (First.IsPackage() && !InAssetData.IsValid()) 
+	{
+		const FString PackageNameStr = Identifiers[0].PackageName.ToString();
+		if ( PackageNameStr.StartsWith(TEXT("/Script")) )// Native Packages (/Script Code)
+		{
+			MainAssetName = PackageNameStr.RightChop(8);
+			AssetTypeName = TEXT("Script");
+		}
+	}
 
 	if (NewIdentifiers.Num() == 1 )
 	{
@@ -246,6 +255,7 @@ void UEdGraphNode_Reference::CacheAssetData(const FAssetData& AssetData)
 			{
 				if ( PackageNameStr.StartsWith(TEXT("/Script")) )
 				{
+					// Used Only in the UI for the Thumbnail
 					CachedAssetData.AssetClassPath = FTopLevelAssetPath(TEXT("/EdGraphNode_Reference"), TEXT("Code"));
 				}
 				else
@@ -254,6 +264,7 @@ void UEdGraphNode_Reference::CacheAssetData(const FAssetData& AssetData)
 					const bool bIsMapPackage = FPlatformFileManager::Get().GetPlatformFile().FileExists(*PotentiallyMapFilename);
 					if ( bIsMapPackage )
 					{
+						// Used Only in the UI for the Thumbnail
 						CachedAssetData.AssetClassPath = TEXT("/Script/Engine.World");
 					}
 				}
diff --git a/Engine/Plugins/Editor/AssetManagerEditor/Source/AssetManagerEditor/Private/ReferenceViewer/SReferenceNode.cpp b/Engine/Plugins/Editor/AssetManagerEditor/Source/AssetManagerEditor/Private/ReferenceViewer/SReferenceNode.cpp
--- a/Engine/Plugins/Editor/AssetManagerEditor/Source/AssetManagerEditor/Private/ReferenceViewer/SReferenceNode.cpp
+++ b/Engine/Plugins/Editor/AssetManagerEditor/Source/AssetManagerEditor/Private/ReferenceViewer/SReferenceNode.cpp
@@ -71,7 +71,6 @@ void SReferenceNode::UpdateGraphNode()
 		IconBrush = GraphNode->GetIconAndTint(IconColor).GetOptionalIcon();
 	}
 
-
 	TSharedRef<SWidget> ThumbnailWidget = SNullWidget::NullWidget;
 	UEdGraphNode_Reference* RefGraphNode = CastChecked<UEdGraphNode_Reference>(GraphNode);
 	bool bIsADuplicate = RefGraphNode->IsADuplicate();
@@ -84,6 +83,7 @@ void SReferenceNode::UpdateGraphNode()
 		FAssetThumbnailConfig ThumbnailConfig;
 		ThumbnailConfig.bAllowFadeIn = RefGraphNode->UsesThumbnail();
 		ThumbnailConfig.bForceGenericThumbnail = !RefGraphNode->UsesThumbnail();
+		ThumbnailConfig.AssetTypeColorOverride = FLinearColor::Transparent;
 
 		ThumbnailWidget =
 			SNew(SBox)
diff --git a/Engine/Source/Editor/UnrealEd/Private/AssetThumbnail.cpp b/Engine/Source/Editor/UnrealEd/Private/AssetThumbnail.cpp
--- a/Engine/Source/Editor/UnrealEd/Private/AssetThumbnail.cpp
+++ b/Engine/Source/Editor/UnrealEd/Private/AssetThumbnail.cpp
@@ -662,7 +662,7 @@ private:
 			return GetDisplayNameForClass( Class, &AssetData );
 		}
 
-		return FText::FromString(AssetClass.ToString());
+		return FText::FromString(AssetClass.GetAssetName().ToString());
 	}
 
 	FText GetAssetDisplayName() const
