Reference Viewer - Expand Overflow Nodes with Double Click

diff --git a/Engine/Plugins/Editor/AssetManagerEditor/Source/AssetManagerEditor/Private/ReferenceViewer/EdGraph_ReferenceViewer.cpp b/Engine/Plugins/Editor/AssetManagerEditor/Source/AssetManagerEditor/Private/ReferenceViewer/EdGraph_ReferenceViewer.cpp
--- a/Engine/Plugins/Editor/AssetManagerEditor/Source/AssetManagerEditor/Private/ReferenceViewer/EdGraph_ReferenceViewer.cpp
+++ b/Engine/Plugins/Editor/AssetManagerEditor/Source/AssetManagerEditor/Private/ReferenceViewer/EdGraph_ReferenceViewer.cpp
@@ -20,6 +20,7 @@ FReferenceNodeInfo::FReferenceNodeInfo(const FAssetIdentifier& InAssetId, bool I
 	: AssetId(InAssetId)
 	, bReferencers(InbReferencers)
 	, OverflowCount(0)
+	, bExpandAllChildren(false)
 	, ChildProvisionSize(0)
 	, PassedFilters(true)
 {}
@@ -330,7 +331,7 @@ void UEdGraph_ReferenceViewer::RecursivelyFilterNodeInfos(const FAssetIdentifier
 
 			if (ChildProvSize > 0)
 			{
-				if (!ExceedsMaxSearchBreadth(Breadth))
+				if (!ExceedsMaxSearchBreadth(Breadth) || InNodeInfos[InAssetId].bExpandAllChildren)
 				{
 					NewProvisionSize += ChildProvSize;
 					Breadth++;
@@ -552,7 +553,7 @@ UEdGraphNode_Reference* UEdGraph_ReferenceViewer::RecursivelyCreateNodes(bool bI
 
 		for (const TPair<FAssetIdentifier, EDependencyPinCategory>& Pair : InNodeInfos[InAssetId].Children)
 		{
-			if (ExceedsMaxSearchBreadth(Breadth))
+			if (ExceedsMaxSearchBreadth(Breadth) && !InNodeInfos[InAssetId].bExpandAllChildren)
 			{
 				break;
 			}
@@ -619,6 +620,21 @@ UEdGraphNode_Reference* UEdGraph_ReferenceViewer::RecursivelyCreateNodes(bool bI
 	return NewNode;
 }
 
+void UEdGraph_ReferenceViewer::ExpandNode(bool bReferencers, const FAssetIdentifier& InAssetIdentifier)
+{
+	if (!bReferencers && DependencyNodeInfos.Contains(InAssetIdentifier))
+	{
+		DependencyNodeInfos[InAssetIdentifier].bExpandAllChildren = true;
+		RefilterGraph();
+	}
+
+	else if (bReferencers && ReferencerNodeInfos.Contains(InAssetIdentifier))
+	{
+		ReferencerNodeInfos[InAssetIdentifier].bExpandAllChildren = true;
+		RefilterGraph();
+	}
+}
+
 const TSharedPtr<FAssetThumbnailPool>& UEdGraph_ReferenceViewer::GetAssetThumbnailPool() const
 {
 	return AssetThumbnailPool;
diff --git a/Engine/Plugins/Editor/AssetManagerEditor/Source/AssetManagerEditor/Private/ReferenceViewer/SReferenceViewer.cpp b/Engine/Plugins/Editor/AssetManagerEditor/Source/AssetManagerEditor/Private/ReferenceViewer/SReferenceViewer.cpp
--- a/Engine/Plugins/Editor/AssetManagerEditor/Source/AssetManagerEditor/Private/ReferenceViewer/SReferenceViewer.cpp
+++ b/Engine/Plugins/Editor/AssetManagerEditor/Source/AssetManagerEditor/Private/ReferenceViewer/SReferenceViewer.cpp
@@ -518,9 +518,46 @@ void SReferenceViewer::SetCurrentRegistrySource(const FAssetManagerEditorRegistr
 
 void SReferenceViewer::OnNodeDoubleClicked(UEdGraphNode* Node)
 {
-	TSet<UObject*> Nodes;
-	Nodes.Add(Node);
-	ReCenterGraphOnNodes( Nodes );
+
+	bool bFoundOverflow = false;
+	if (UEdGraphNode_Reference* ReferenceNode = Cast<UEdGraphNode_Reference>(Node))
+	{
+		// Overflow nodes have no identifiers
+		if (!ReferenceNode->GetIdentifier().IsValid()) 
+		{
+			if (ReferenceNode->GetReferencerPin()->LinkedTo.Num() > 0)
+			{
+				if (UEdGraphNode* ParentNode = ReferenceNode->GetReferencerPin()->LinkedTo[0]->GetOwningNode())
+				{
+					if (UEdGraphNode_Reference* ParentReferenceNode = Cast<UEdGraphNode_Reference>(ParentNode))
+					{
+						FAssetIdentifier ParentID = ParentReferenceNode->GetIdentifier();
+						GraphObj->ExpandNode(false, ParentID);
+						bFoundOverflow = true;
+					}
+				}
+			}
+			else if (ReferenceNode->GetDependencyPin()->LinkedTo.Num() > 0)
+			{
+				if (UEdGraphNode* ParentNode = ReferenceNode->GetDependencyPin()->LinkedTo[0]->GetOwningNode())
+				{
+					if (UEdGraphNode_Reference* ParentReferenceNode = Cast<UEdGraphNode_Reference>(ParentNode))
+					{
+						FAssetIdentifier ParentID = ParentReferenceNode->GetIdentifier();
+						GraphObj->ExpandNode(true, ParentID);
+						bFoundOverflow = true;
+					}
+				}
+			}
+		}
+	}
+
+	if (!bFoundOverflow)
+	{
+		TSet<UObject*> Nodes;
+		Nodes.Add(Node);
+		ReCenterGraphOnNodes( Nodes );
+	}
 }
 
 void SReferenceViewer::RebuildGraph()
diff --git a/Engine/Plugins/Editor/AssetManagerEditor/Source/AssetManagerEditor/Public/ReferenceViewer/EdGraph_ReferenceViewer.h b/Engine/Plugins/Editor/AssetManagerEditor/Source/AssetManagerEditor/Public/ReferenceViewer/EdGraph_ReferenceViewer.h
--- a/Engine/Plugins/Editor/AssetManagerEditor/Source/AssetManagerEditor/Public/ReferenceViewer/EdGraph_ReferenceViewer.h
+++ b/Engine/Plugins/Editor/AssetManagerEditor/Source/AssetManagerEditor/Public/ReferenceViewer/EdGraph_ReferenceViewer.h
@@ -37,6 +37,9 @@ struct FReferenceNodeInfo
 
 	int32 OverflowCount;
 
+	// Denote when all children have been manually expanded and the breadth limit should be ignored
+	bool bExpandAllChildren;
+
 	FReferenceNodeInfo(const FAssetIdentifier& InAssetId, bool InbReferencers);
 
 	bool IsFirstParent(const FAssetIdentifier& InParentId) const;
@@ -134,6 +137,8 @@ private:
 		bool bIsRoot = false
 	);
 
+	void ExpandNode(bool bReferencers, const FAssetIdentifier& InAssetIdentifier);
+
 	/** Removes all nodes from the graph */
 	void RemoveAllNodes();
 
