Reference Viewer - Breadth Limit Indicators

diff --git a/Engine/Plugins/Editor/AssetManagerEditor/Source/AssetManagerEditor/Private/ReferenceViewer/EdGraphNode_Reference.cpp b/Engine/Plugins/Editor/AssetManagerEditor/Source/AssetManagerEditor/Private/ReferenceViewer/EdGraphNode_Reference.cpp
--- a/Engine/Plugins/Editor/AssetManagerEditor/Source/AssetManagerEditor/Private/ReferenceViewer/EdGraphNode_Reference.cpp
+++ b/Engine/Plugins/Editor/AssetManagerEditor/Source/AssetManagerEditor/Private/ReferenceViewer/EdGraphNode_Reference.cpp
@@ -26,6 +26,8 @@ UEdGraphNode_Reference::UEdGraphNode_Reference(const FObjectInitializer& ObjectI
 	bUsesThumbnail = false;
 	bAllowThumbnail = true;
 	AssetTypeColor = FLinearColor(0.55f, 0.55f, 0.55f);
+	bIsFiltered = false;
+	bIsOverflow = false;
 }
 
 void UEdGraphNode_Reference::SetupReferenceNode(const FIntPoint& NodeLoc, const TArray<FAssetIdentifier>& NewIdentifiers, const FAssetData& InAssetData, bool bInAllowThumbnail, bool bInIsADuplicate)
@@ -110,6 +112,8 @@ void UEdGraphNode_Reference::SetReferenceNodeCollapsed(const FIntPoint& NodeLoc,
 	Identifiers.Empty();
 	bIsCollapsed = true;
 	bUsesThumbnail = false;
+	bIsOverflow = true;
+	AssetBrush = FSlateIcon(FAppStyle::GetAppStyleSetName(), "Icons.WarningWithColor");
 
 	NodeTitle = FText::Format( LOCTEXT("ReferenceNodeCollapsedTitle", "{0} Collapsed nodes"), FText::AsNumber(InNumReferencesExceedingMax));
 	CacheAssetData(FAssetData());
@@ -200,7 +204,7 @@ FText UEdGraphNode_Reference::GetTooltipText() const
 
 FSlateIcon UEdGraphNode_Reference::GetIconAndTint(FLinearColor& OutColor) const
 {
-	OutColor = AssetTypeColor;
+	OutColor = bIsOverflow ? FLinearColor::White : AssetTypeColor;
 	return AssetBrush;
 }
 
diff --git a/Engine/Plugins/Editor/AssetManagerEditor/Source/AssetManagerEditor/Private/ReferenceViewer/EdGraph_ReferenceViewer.cpp b/Engine/Plugins/Editor/AssetManagerEditor/Source/AssetManagerEditor/Private/ReferenceViewer/EdGraph_ReferenceViewer.cpp
--- a/Engine/Plugins/Editor/AssetManagerEditor/Source/AssetManagerEditor/Private/ReferenceViewer/EdGraph_ReferenceViewer.cpp
+++ b/Engine/Plugins/Editor/AssetManagerEditor/Source/AssetManagerEditor/Private/ReferenceViewer/EdGraph_ReferenceViewer.cpp
@@ -274,6 +274,8 @@ UEdGraphNode_Reference* UEdGraph_ReferenceViewer::RefilterGraph()
 {
 	RemoveAllNodes();
 	UEdGraphNode_Reference* RootNode = NULL;
+
+	bBreadthLimitReached = false;
 	if (CurrentGraphRootIdentifiers.Num() > 0)
 	{
 		FAssetIdentifier FirstGraphRootIdentifier = CurrentGraphRootIdentifiers[0];
@@ -350,6 +352,7 @@ void UEdGraph_ReferenceViewer::RecursivelyFilterNodeInfos(const FAssetIdentifier
 	if (InNodeInfos[InAssetId].OverflowCount > 0)
 	{
 		NewProvisionSize++;
+		bBreadthLimitReached = true;
 	}
 
 	bool PassedAllFilters = FilterCollection && Settings->GetFiltersEnabled() ? FilterCollection->PassesAllFilters(InNodeInfos[InAssetId]) : true;
diff --git a/Engine/Plugins/Editor/AssetManagerEditor/Source/AssetManagerEditor/Private/ReferenceViewer/SReferenceNode.cpp b/Engine/Plugins/Editor/AssetManagerEditor/Source/AssetManagerEditor/Private/ReferenceViewer/SReferenceNode.cpp
--- a/Engine/Plugins/Editor/AssetManagerEditor/Source/AssetManagerEditor/Private/ReferenceViewer/SReferenceNode.cpp
+++ b/Engine/Plugins/Editor/AssetManagerEditor/Source/AssetManagerEditor/Private/ReferenceViewer/SReferenceNode.cpp
@@ -147,11 +147,13 @@ void SReferenceNode::UpdateGraphNode()
 						]
 						+ SHorizontalBox::Slot()
 						.FillWidth(1.0)
+						.VAlign(VAlign_Center)
 						[
 							SNew(SVerticalBox)
 							+SVerticalBox::Slot()
 							.AutoHeight()
 							.Padding(FMargin(0.f))
+							.VAlign(VAlign_Center)
 							[
 								SAssignNew(InlineEditableText, SInlineEditableTextBlock)
 								.Style( FReferenceViewerStyle::Get(), "Graph.Node.NodeTitleInlineEditableText" )
diff --git a/Engine/Plugins/Editor/AssetManagerEditor/Source/AssetManagerEditor/Private/ReferenceViewer/SReferenceViewer.cpp b/Engine/Plugins/Editor/AssetManagerEditor/Source/AssetManagerEditor/Private/ReferenceViewer/SReferenceViewer.cpp
--- a/Engine/Plugins/Editor/AssetManagerEditor/Source/AssetManagerEditor/Private/ReferenceViewer/SReferenceViewer.cpp
+++ b/Engine/Plugins/Editor/AssetManagerEditor/Source/AssetManagerEditor/Private/ReferenceViewer/SReferenceViewer.cpp
@@ -323,6 +323,7 @@ void SReferenceViewer::Construct(const FArguments& InArgs)
 
 						+SHorizontalBox::Slot()
 						.VAlign(VAlign_Center)
+						.AutoWidth()
 						.Padding(2.f)
 						[
 							SNew(STextBlock)
@@ -333,6 +334,18 @@ void SReferenceViewer::Construct(const FArguments& InArgs)
 															UICommands.SetBreadth->GetInputText().ToUpper()))
 						]
 
+						+SHorizontalBox::Slot()
+						.VAlign(VAlign_Center)
+						.HAlign(HAlign_Left)
+						.FillWidth(1.0)
+						.Padding(2.f, 0.f, 8.f, 0.f)
+						[
+							SNew(SImage)
+							.ToolTipText(LOCTEXT("BreadthLimitToolTip", "The Breadth Limit was reached."))
+							.Image(FAppStyle::GetBrush("Icons.WarningWithColor"))
+							.Visibility_Lambda([this] { return GraphObj && GraphObj->BreadthLimitExceeded() ? EVisibility::Visible : EVisibility::Hidden; })
+						]
+
 						+SHorizontalBox::Slot()
 						.AutoWidth()
 						.VAlign(VAlign_Center)
@@ -523,7 +536,7 @@ void SReferenceViewer::OnNodeDoubleClicked(UEdGraphNode* Node)
 	if (UEdGraphNode_Reference* ReferenceNode = Cast<UEdGraphNode_Reference>(Node))
 	{
 		// Overflow nodes have no identifiers
-		if (!ReferenceNode->GetIdentifier().IsValid()) 
+		if (ReferenceNode->IsOverflow())
 		{
 			if (ReferenceNode->GetReferencerPin()->LinkedTo.Num() > 0)
 			{
diff --git a/Engine/Plugins/Editor/AssetManagerEditor/Source/AssetManagerEditor/Public/ReferenceViewer/EdGraphNode_Reference.h b/Engine/Plugins/Editor/AssetManagerEditor/Source/AssetManagerEditor/Public/ReferenceViewer/EdGraphNode_Reference.h
--- a/Engine/Plugins/Editor/AssetManagerEditor/Source/AssetManagerEditor/Public/ReferenceViewer/EdGraphNode_Reference.h
+++ b/Engine/Plugins/Editor/AssetManagerEditor/Source/AssetManagerEditor/Public/ReferenceViewer/EdGraphNode_Reference.h
@@ -51,6 +51,8 @@ class ASSETMANAGEREDITOR_API UEdGraphNode_Reference : public UEdGraphNode
 	void SetIsFiltered(bool bInFiltered);
 	bool GetIsFiltered() const;
 
+	bool IsOverflow() const { return bIsOverflow; }
+
 	FAssetData GetAssetData() const;
 
 	UEdGraphPin* GetDependencyPin();
@@ -72,6 +74,7 @@ private:
 	bool bIsCollapsed;
 	bool bIsADuplicate;
 	bool bIsFiltered;
+	bool bIsOverflow;
 
 	FAssetData CachedAssetData;
 	FLinearColor AssetTypeColor;
diff --git a/Engine/Plugins/Editor/AssetManagerEditor/Source/AssetManagerEditor/Public/ReferenceViewer/EdGraph_ReferenceViewer.h b/Engine/Plugins/Editor/AssetManagerEditor/Source/AssetManagerEditor/Public/ReferenceViewer/EdGraph_ReferenceViewer.h
--- a/Engine/Plugins/Editor/AssetManagerEditor/Source/AssetManagerEditor/Public/ReferenceViewer/EdGraph_ReferenceViewer.h
+++ b/Engine/Plugins/Editor/AssetManagerEditor/Source/AssetManagerEditor/Public/ReferenceViewer/EdGraph_ReferenceViewer.h
@@ -104,6 +104,9 @@ public:
 	/* Returns a set of unique asset types as UClass* */
 	const TSet<FTopLevelAssetPath>& GetAssetTypes() const { return CurrentClasses; }
 
+	/* Returns true if the current graph has overflow nodes */
+	bool BreadthLimitExceeded() const { return bBreadthLimitReached; };
+
 private:
 	void SetReferenceViewer(TSharedPtr<SReferenceViewer> InViewer);
 	UEdGraphNode_Reference* ConstructNodes(const TArray<FAssetIdentifier>& GraphRootIdentifiers, const FIntPoint& GraphRootOrigin);
@@ -160,6 +163,9 @@ private:
 
 	int32 MaxSearchBreadth;
 
+	/** Stores if the breadth limit was reached on the last refilter*/
+	bool bBreadthLimitReached;
+
 	/** Current collection filter. NAME_None for no filter */
 	FName CurrentCollectionFilter;
 
