Fixed conflicting LOCTEXT warning in Reference Viewer

diff --git a/Engine/Plugins/Editor/AssetManagerEditor/Source/AssetManagerEditor/Private/ReferenceViewer/SReferenceViewer.cpp b/Engine/Plugins/Editor/AssetManagerEditor/Source/AssetManagerEditor/Private/ReferenceViewer/SReferenceViewer.cpp
--- a/Engine/Plugins/Editor/AssetManagerEditor/Source/AssetManagerEditor/Private/ReferenceViewer/SReferenceViewer.cpp
+++ b/Engine/Plugins/Editor/AssetManagerEditor/Source/AssetManagerEditor/Private/ReferenceViewer/SReferenceViewer.cpp
@@ -341,7 +341,7 @@ void SReferenceViewer::Construct(const FArguments& InArgs)
 						.Padding(2.f, 0.f, 8.f, 0.f)
 						[
 							SNew(SImage)
-							.ToolTipText(LOCTEXT("BreadthLimitToolTip", "The Breadth Limit was reached."))
+							.ToolTipText(LOCTEXT("BreadthLimitReachedToolTip", "The Breadth Limit was reached."))
 							.Image(FAppStyle::GetBrush("Icons.WarningWithColor"))
 							.Visibility_Lambda([this] { return GraphObj && GraphObj->BreadthLimitExceeded() ? EVisibility::Visible : EVisibility::Hidden; })
 						]
