Reference Viewer - Ensure Management Database is updated on initialization if the Show Management Referenes is enabled.

diff --git a/Engine/Plugins/Editor/AssetManagerEditor/Source/AssetManagerEditor/Private/ReferenceViewer/SReferenceViewer.cpp b/Engine/Plugins/Editor/AssetManagerEditor/Source/AssetManagerEditor/Private/ReferenceViewer/SReferenceViewer.cpp
--- a/Engine/Plugins/Editor/AssetManagerEditor/Source/AssetManagerEditor/Private/ReferenceViewer/SReferenceViewer.cpp
+++ b/Engine/Plugins/Editor/AssetManagerEditor/Source/AssetManagerEditor/Private/ReferenceViewer/SReferenceViewer.cpp
@@ -505,6 +505,11 @@ void SReferenceViewer::SetGraphRootIdentifiers(const TArray<FAssetIdentifier>& N
 		Settings->SetCompactModeEnabled(ReferenceViewerParams.bCompactMode.GetValue());
 	}
 
+	if (Settings->IsShowManagementReferences())
+	{
+		UAssetManager::Get().UpdateManagementDatabase();
+	}
+
 	RebuildGraph();
 
 	// Zoom once this frame to make sure widgets are visible, then zoom again so size is correct
