Reference Viewer Improves Search Results

- Will now include PrimaryAssetIds in search results.
- Filtering search results will properly include *all* matching nodes even if in-between nodes do not match the filter.  Unmatched waypoint nodes will show ghosted (same as they do for asset type filters)

diff --git a/Engine/Plugins/Editor/AssetManagerEditor/Source/AssetManagerEditor/Private/ReferenceViewer/EdGraph_ReferenceViewer.cpp b/Engine/Plugins/Editor/AssetManagerEditor/Source/AssetManagerEditor/Private/ReferenceViewer/EdGraph_ReferenceViewer.cpp
--- a/Engine/Plugins/Editor/AssetManagerEditor/Source/AssetManagerEditor/Private/ReferenceViewer/EdGraph_ReferenceViewer.cpp
+++ b/Engine/Plugins/Editor/AssetManagerEditor/Source/AssetManagerEditor/Private/ReferenceViewer/EdGraph_ReferenceViewer.cpp
@@ -276,13 +276,13 @@ UEdGraphNode_Reference* UEdGraph_ReferenceViewer::RefilterGraph()
 	UEdGraphNode_Reference* RootNode = NULL;
 
 	bBreadthLimitReached = false;
-	if (CurrentGraphRootIdentifiers.Num() > 0)
+	if (CurrentGraphRootIdentifiers.Num() > 0 && (!ReferencerNodeInfos.IsEmpty() || !DependencyNodeInfos.IsEmpty()))
 	{
 		FAssetIdentifier FirstGraphRootIdentifier = CurrentGraphRootIdentifiers[0];
 
 		// Create the root node
-		bool bRootIsDuplicated = (Settings->IsShowDependencies() && DependencyNodeInfos[FirstGraphRootIdentifier].IsADuplicate()) || 
-								  (Settings->IsShowReferencers() && ReferencerNodeInfos[FirstGraphRootIdentifier].IsADuplicate());
+		bool bRootIsDuplicated = (Settings->IsShowDependencies() && DependencyNodeInfos.Contains(FirstGraphRootIdentifier) && DependencyNodeInfos[FirstGraphRootIdentifier].IsADuplicate()) || 
+								  (Settings->IsShowReferencers() && ReferencerNodeInfos.Contains(FirstGraphRootIdentifier) && ReferencerNodeInfos[FirstGraphRootIdentifier].IsADuplicate());
 		const FReferenceNodeInfo& NodeInfo = Settings->IsShowReferencers() ? ReferencerNodeInfos[FirstGraphRootIdentifier] : DependencyNodeInfos[FirstGraphRootIdentifier];
 		RootNode = CreateReferenceNode();
 		RootNode->SetupReferenceNode(CurrentGraphRootOrigin, CurrentGraphRootIdentifiers, NodeInfo.AssetData, /*bInAllowThumbnail = */ !Settings->IsCompactMode(), /*bIsDuplicate*/ bRootIsDuplicated);
@@ -355,7 +355,10 @@ void UEdGraph_ReferenceViewer::RecursivelyFilterNodeInfos(const FAssetIdentifier
 		bBreadthLimitReached = true;
 	}
 
-	bool PassedAllFilters = FilterCollection && Settings->GetFiltersEnabled() ? FilterCollection->PassesAllFilters(InNodeInfos[InAssetId]) : true;
+	bool PassedAssetTypeFilter = FilterCollection && Settings->GetFiltersEnabled() ? FilterCollection->PassesAllFilters(InNodeInfos[InAssetId]) : true;
+	bool PassedSearchTextFilter = IsAssetPassingSearchTextFilter(InAssetId);
+
+	bool PassedAllFilters = PassedAssetTypeFilter && PassedSearchTextFilter;	
 
 	InNodeInfos[InAssetId].ChildProvisionSize = NewProvisionSize > 0 ? NewProvisionSize : (PassedAllFilters ? 1 : 0);
 	InNodeInfos[InAssetId].PassedFilters = PassedAllFilters;
@@ -457,10 +460,16 @@ bool UEdGraph_ReferenceViewer::IsPackageIdentifierPassingFilter(const FAssetIden
 			return false;
 		}
 
-		if (Settings->IsShowFilteredPackagesOnly() && IsPackageNamePassingFilterCallback.IsSet() && !(*IsPackageNamePassingFilterCallback)(InAssetIdentifier.PackageName))
-		{
-			return false;
-		}
+	}
+
+	return true;
+}
+
+bool UEdGraph_ReferenceViewer::IsAssetPassingSearchTextFilter(const FAssetIdentifier& InAssetIdentifier) const
+{
+	if (Settings->IsShowFilteredPackagesOnly() && IsAssetIdentifierPassingSearchFilterCallback.IsSet() && !(*IsAssetIdentifierPassingSearchFilterCallback)(InAssetIdentifier))
+	{
+		return false;
 	}
 
 	return true;
diff --git a/Engine/Plugins/Editor/AssetManagerEditor/Source/AssetManagerEditor/Private/ReferenceViewer/SReferenceViewer.cpp b/Engine/Plugins/Editor/AssetManagerEditor/Source/AssetManagerEditor/Private/ReferenceViewer/SReferenceViewer.cpp
--- a/Engine/Plugins/Editor/AssetManagerEditor/Source/AssetManagerEditor/Private/ReferenceViewer/SReferenceViewer.cpp
+++ b/Engine/Plugins/Editor/AssetManagerEditor/Source/AssetManagerEditor/Private/ReferenceViewer/SReferenceViewer.cpp
@@ -43,33 +43,12 @@
 
 #define LOCTEXT_NAMESPACE "ReferenceViewer"
 
-bool IsPackageNamePassingFilter(FName InPackageName, const TArray<FString>& InSearchWords)
+bool IsAssetIdentifierPassingSearchTextFilter(const FAssetIdentifier& InNode, const TArray<FString>& InSearchWords)
 {
-	// package name must match all words
+	FString NodeString = InNode.ToString();
 	for (const FString& Word : InSearchWords)
 	{
-		if (!InPackageName.ToString().Contains(Word))
-		{
-			return false;
-		}
-	}
-
-	return true;
-}
-
-bool IsReferenceNodePassingFilter(const UEdGraphNode_Reference& InNode, const TArray<FString>& InSearchWords)
-{
-	if (InSearchWords.Num() == 0)
-	{
-		return true;
-	}
-
-	TArray<FName> NodePackageNames;
-	InNode.GetAllPackageNames(NodePackageNames);
-
-	for (const FName& PackageName : NodePackageNames)
-	{
-		if (!IsPackageNamePassingFilter(PackageName, InSearchWords))
+		if (!NodeString.Contains(Word))
 		{
 			return false;
 		}
@@ -497,7 +476,7 @@ void SReferenceViewer::SetGraphRootIdentifiers(const TArray<FAssetIdentifier>& N
 	{
 		Settings->SetShowFilteredPackagesOnlyEnabled(ReferenceViewerParams.bShowFilteredPackagesOnly.GetValue());
 	}
-	UpdateIsPassingFilterPackageCallback();
+	UpdateIsPassingSearchFilterCallback();
 
 	bShowCompactMode = ReferenceViewerParams.bShowCompactMode;
 	if (ReferenceViewerParams.bCompactMode.IsSet())
@@ -987,7 +966,7 @@ bool SReferenceViewer::IsShowHardReferencesChecked() const
 void SReferenceViewer::OnShowFilteredPackagesOnlyChanged()
 {
 	Settings->SetShowFilteredPackagesOnlyEnabled(!Settings->IsShowFilteredPackagesOnly());
-	UpdateIsPassingFilterPackageCallback();
+	UpdateIsPassingSearchFilterCallback();
 }
 
 
@@ -996,11 +975,11 @@ bool SReferenceViewer::IsShowFilteredPackagesOnlyChecked() const
 	return Settings->IsShowFilteredPackagesOnly();
 }
 
-void SReferenceViewer::UpdateIsPassingFilterPackageCallback()
+void SReferenceViewer::UpdateIsPassingSearchFilterCallback()
 {
 	if (GraphObj)
 	{
-		TOptional<UEdGraph_ReferenceViewer::FIsPackageNamePassingFilterCallback> IsAssetPassingFilterCallback;
+		TOptional<UEdGraph_ReferenceViewer::FIsAssetIdentifierPassingSearchFilterCallback> IsAssetIdentifierPassingSearchFilterCallback;
 		FString SearchString = SearchBox->GetText().ToString();
 		TArray<FString> SearchWords;
 		SearchString.ParseIntoArrayWS(SearchWords);
@@ -1009,13 +988,13 @@ void SReferenceViewer::UpdateIsPassingFilterPackageCallback()
 			{
 				if (SearchWords.Num() > 0)
 				{
-					IsAssetPassingFilterCallback = [=](FName InName) { return IsPackageNamePassingFilter(InName, SearchWords); };
+					IsAssetIdentifierPassingSearchFilterCallback = [=](const FAssetIdentifier& InAssetIdentifier) { return IsAssetIdentifierPassingSearchTextFilter(InAssetIdentifier, SearchWords); };
 				}
 			}
 
-			GraphObj->SetIsPackageNamePassingFilterCallback(IsAssetPassingFilterCallback);
+			GraphObj->SetIsAssetIdentifierPassingSearchFilterCallback(IsAssetIdentifierPassingSearchFilterCallback);
 		}
-		RebuildGraph();
+		GraphObj->RefilterGraph();
 	}
 }
 
@@ -1865,7 +1844,7 @@ void SReferenceViewer::HandleOnSearchTextChanged(const FText& SearchText)
 
 	GraphEditorPtr->ClearSelectionSet();
 
-	UpdateIsPassingFilterPackageCallback();
+	UpdateIsPassingSearchFilterCallback();
 
 	if (SearchText.IsEmpty())
 	{
@@ -1881,7 +1860,7 @@ void SReferenceViewer::HandleOnSearchTextChanged(const FText& SearchText)
 
 	for (UEdGraphNode_Reference* Node : AllNodes)
 	{
-		if (IsReferenceNodePassingFilter(*Node, SearchWords))
+		if (IsAssetIdentifierPassingSearchTextFilter(Node->GetIdentifier(), SearchWords))
 		{
 			GraphEditorPtr->SetNodeSelection(Node, true);
 		}
diff --git a/Engine/Plugins/Editor/AssetManagerEditor/Source/AssetManagerEditor/Private/ReferenceViewer/SReferenceViewer.h b/Engine/Plugins/Editor/AssetManagerEditor/Source/AssetManagerEditor/Private/ReferenceViewer/SReferenceViewer.h
--- a/Engine/Plugins/Editor/AssetManagerEditor/Source/AssetManagerEditor/Private/ReferenceViewer/SReferenceViewer.h
+++ b/Engine/Plugins/Editor/AssetManagerEditor/Source/AssetManagerEditor/Private/ReferenceViewer/SReferenceViewer.h
@@ -137,7 +137,7 @@ private:
 
 	void OnShowFilteredPackagesOnlyChanged();
 	bool IsShowFilteredPackagesOnlyChecked() const;
-	void UpdateIsPassingFilterPackageCallback();
+	void UpdateIsPassingSearchFilterCallback();
 
 	void OnCompactModeChanged();
 	bool IsCompactModeChecked() const;
diff --git a/Engine/Plugins/Editor/AssetManagerEditor/Source/AssetManagerEditor/Public/ReferenceViewer/EdGraph_ReferenceViewer.h b/Engine/Plugins/Editor/AssetManagerEditor/Source/AssetManagerEditor/Public/ReferenceViewer/EdGraph_ReferenceViewer.h
--- a/Engine/Plugins/Editor/AssetManagerEditor/Source/AssetManagerEditor/Public/ReferenceViewer/EdGraph_ReferenceViewer.h
+++ b/Engine/Plugins/Editor/AssetManagerEditor/Source/AssetManagerEditor/Public/ReferenceViewer/EdGraph_ReferenceViewer.h
@@ -89,8 +89,8 @@ public:
 	 * the presentation or filtering */
 	class UEdGraphNode_Reference* RefilterGraph();
 
-	using FIsPackageNamePassingFilterCallback = TFunction<bool(FName)>;
-	void SetIsPackageNamePassingFilterCallback(const TOptional<FIsPackageNamePassingFilterCallback>& InIsPackageNamePassingFilterCallback) { IsPackageNamePassingFilterCallback = InIsPackageNamePassingFilterCallback; }
+	using FIsAssetIdentifierPassingSearchFilterCallback = TFunction<bool(const FAssetIdentifier&)>;
+	void SetIsAssetIdentifierPassingSearchFilterCallback(const TOptional<FIsAssetIdentifierPassingSearchFilterCallback>& InIsAssetIdentifierPassingSearchFilterCallback) { IsAssetIdentifierPassingSearchFilterCallback = InIsAssetIdentifierPassingSearchFilterCallback; }
 
 	FName GetCurrentCollectionFilter() const;
 	void SetCurrentCollectionFilter(FName NewFilter);
@@ -150,6 +150,7 @@ private:
 
 	void GetSortedLinks(const TArray<FAssetIdentifier>& Identifiers, bool bReferencers, const FAssetManagerDependencyQuery& Query, TMap<FAssetIdentifier, EDependencyPinCategory>& OutLinks) const;
 	bool IsPackageIdentifierPassingFilter(const FAssetIdentifier& InAssetIdentifier) const;
+	bool IsAssetPassingSearchTextFilter(const FAssetIdentifier& InAssetIdentifier) const;
 
 private:
 	/** Pool for maintaining and rendering thumbnails */
@@ -179,7 +180,7 @@ private:
 	TMap<FAssetIdentifier, FReferenceNodeInfo> ReferencerNodeInfos;
 	TMap<FAssetIdentifier, FReferenceNodeInfo> DependencyNodeInfos;
 
-	TOptional<FIsPackageNamePassingFilterCallback> IsPackageNamePassingFilterCallback;
+	TOptional<FIsAssetIdentifierPassingSearchFilterCallback> IsAssetIdentifierPassingSearchFilterCallback;
 
 	/** List of packages the current collection filter allows */
 	TSet<FName> CurrentCollectionPackages;
