#UE Fix reference viewer not showing all references when centered on multiple assets

diff --git a/Engine/Plugins/Editor/AssetManagerEditor/Source/AssetManagerEditor/Private/ReferenceViewer/EdGraph_ReferenceViewer.cpp b/Engine/Plugins/Editor/AssetManagerEditor/Source/AssetManagerEditor/Private/ReferenceViewer/EdGraph_ReferenceViewer.cpp
--- a/Engine/Plugins/Editor/AssetManagerEditor/Source/AssetManagerEditor/Private/ReferenceViewer/EdGraph_ReferenceViewer.cpp
+++ b/Engine/Plugins/Editor/AssetManagerEditor/Source/AssetManagerEditor/Private/ReferenceViewer/EdGraph_ReferenceViewer.cpp
@@ -196,12 +196,12 @@ UEdGraphNode_Reference* UEdGraph_ReferenceViewer::ConstructNodes(const TArray<FA
 		TMap<FAssetIdentifier, FReferenceNodeInfo> NewReferenceNodeInfos;
 		FReferenceNodeInfo& RootNodeInfo = NewReferenceNodeInfos.FindOrAdd( GraphRootIdentifiers[0], FReferenceNodeInfo(GraphRootIdentifiers[0], true));
 		RootNodeInfo.Parents.Emplace(FAssetIdentifier(NAME_None));
-		RecursivelyPopulateNodeInfos(true, GraphRootIdentifiers[0], NewReferenceNodeInfos, 0, Settings->GetSearchReferencerDepthLimit());
+		RecursivelyPopulateNodeInfos(true, GraphRootIdentifiers, NewReferenceNodeInfos, 0, Settings->GetSearchReferencerDepthLimit());
 
 		TMap<FAssetIdentifier, FReferenceNodeInfo> NewDependencyNodeInfos;
 		FReferenceNodeInfo& DRootNodeInfo = NewDependencyNodeInfos.FindOrAdd( GraphRootIdentifiers[0], FReferenceNodeInfo(GraphRootIdentifiers[0], false));
 		DRootNodeInfo.Parents.Emplace(FAssetIdentifier(NAME_None));
-		RecursivelyPopulateNodeInfos(false, GraphRootIdentifiers[0], NewDependencyNodeInfos, 0, Settings->GetSearchDependencyDepthLimit());
+		RecursivelyPopulateNodeInfos(false, GraphRootIdentifiers, NewDependencyNodeInfos, 0, Settings->GetSearchDependencyDepthLimit());
 
 		TSet<FName> AllPackageNames;
 		auto AddPackage = [](const FAssetIdentifier& AssetId, TSet<FName>& PackageNames)
@@ -476,17 +476,15 @@ bool UEdGraph_ReferenceViewer::IsAssetPassingSearchTextFilter(const FAssetIdenti
 }
 
 void
-UEdGraph_ReferenceViewer::RecursivelyPopulateNodeInfos(bool bInReferencers, const FAssetIdentifier& InAssetId, TMap<FAssetIdentifier, FReferenceNodeInfo>& InNodeInfos, int32 InCurrentDepth, int32 InMaxDepth)
+UEdGraph_ReferenceViewer::RecursivelyPopulateNodeInfos(bool bInReferencers, const TArray<FAssetIdentifier>& Identifiers, TMap<FAssetIdentifier, FReferenceNodeInfo>& InNodeInfos, int32 InCurrentDepth, int32 InMaxDepth)
 {
+	check(Identifiers.Num() > 0);
 	int32 ProvisionSize = 0;
-
+	const FAssetIdentifier& InAssetId = Identifiers[0];
 	if (InMaxDepth > 0 && InCurrentDepth < InMaxDepth)
 	{
-		TArray<FAssetIdentifier> AssetIds;
-		AssetIds.Add(InAssetId);
-
 		TMap<FAssetIdentifier, EDependencyPinCategory> ReferenceLinks;
-		GetSortedLinks(AssetIds, bInReferencers, GetReferenceSearchFlags(false), ReferenceLinks);
+		GetSortedLinks(Identifiers, bInReferencers, GetReferenceSearchFlags(false), ReferenceLinks);
 
 		InNodeInfos[InAssetId].Children.Reserve(ReferenceLinks.Num());
 		for (const TPair<FAssetIdentifier, EDependencyPinCategory>& Pair : ReferenceLinks)
@@ -498,7 +496,7 @@ UEdGraph_ReferenceViewer::RecursivelyPopulateNodeInfos(bool bInReferencers, cons
 				InNodeInfos[ChildId].Parents.Emplace(InAssetId);
 				InNodeInfos[InAssetId].Children.Emplace(Pair);
 
-				RecursivelyPopulateNodeInfos(bInReferencers, ChildId, InNodeInfos, InCurrentDepth + 1, InMaxDepth);
+				RecursivelyPopulateNodeInfos(bInReferencers, { ChildId }, InNodeInfos, InCurrentDepth + 1, InMaxDepth);
 				ProvisionSize += InNodeInfos[ChildId].ProvisionSize(InAssetId);
 			}
 
diff --git a/Engine/Plugins/Editor/AssetManagerEditor/Source/AssetManagerEditor/Public/ReferenceViewer/EdGraph_ReferenceViewer.h b/Engine/Plugins/Editor/AssetManagerEditor/Source/AssetManagerEditor/Public/ReferenceViewer/EdGraph_ReferenceViewer.h
--- a/Engine/Plugins/Editor/AssetManagerEditor/Source/AssetManagerEditor/Public/ReferenceViewer/EdGraph_ReferenceViewer.h
+++ b/Engine/Plugins/Editor/AssetManagerEditor/Source/AssetManagerEditor/Public/ReferenceViewer/EdGraph_ReferenceViewer.h
@@ -119,7 +119,7 @@ private:
 	UEdGraphNode_Reference* CreateReferenceNode();
 
 	/* Generates a NodeInfo structure then used to generate and layout the graph nodes */
-	void RecursivelyPopulateNodeInfos(bool bReferencers, const FAssetIdentifier& AssetId, TMap<FAssetIdentifier, FReferenceNodeInfo>& NodeInfos, int32 CurrentDepth, int32 MaxDepth);
+	void RecursivelyPopulateNodeInfos(bool bReferencers, const TArray<FAssetIdentifier>& Identifiers, TMap<FAssetIdentifier, FReferenceNodeInfo>& NodeInfos, int32 CurrentDepth, int32 MaxDepth);
 
 	/* Marks up the NodeInfos with updated filter information and provision sizes */
 	void RecursivelyFilterNodeInfos(const FAssetIdentifier& InAssetId, TMap<FAssetIdentifier, FReferenceNodeInfo>& NodeInfos, int32 CurrentDepth, int32 MaxDepth);
