Clarifying presentation of Name References and Management References in the Reference Viewer

diff --git a/Engine/Config/BaseEditorPerProjectUserSettings.ini b/Engine/Config/BaseEditorPerProjectUserSettings.ini
--- a/Engine/Config/BaseEditorPerProjectUserSettings.ini
+++ b/Engine/Config/BaseEditorPerProjectUserSettings.ini
@@ -961,7 +961,7 @@ bIsShowHardReferences=True
 bIsShowEditorOnlyReferences=True
 bIsShowManagementReferences=False
 bIsShowSearchableNames=False
-bIsShowNativePackages=False
+bIsShowCodePackages=False
 bIsShowDuplicates=True
 bIsShowFilteredPackagesOnly=False
 bIsCompactMode=False
diff --git a/Engine/Plugins/Editor/AssetManagerEditor/Source/AssetManagerEditor/Private/AssetManagerEditorCommands.cpp b/Engine/Plugins/Editor/AssetManagerEditor/Source/AssetManagerEditor/Private/AssetManagerEditorCommands.cpp
--- a/Engine/Plugins/Editor/AssetManagerEditor/Source/AssetManagerEditor/Private/AssetManagerEditorCommands.cpp
+++ b/Engine/Plugins/Editor/AssetManagerEditor/Source/AssetManagerEditor/Private/AssetManagerEditorCommands.cpp
@@ -43,11 +43,11 @@ void FAssetManagerEditorCommands::RegisterCommands()
 	UI_COMMAND(ShowEditorOnlyReferences, "Show EditorOnly References","Toggles visibility of EditorOnly References", EUserInterfaceActionType::ToggleButton, FInputChord(EKeys::E));
 
 	UI_COMMAND(ShowManagementReferences, "Show Management References","Toggles visibility of Management References (i.e. PrimaryAssetIDs)", EUserInterfaceActionType::ToggleButton, FInputChord(EKeys::M));
-	UI_COMMAND(ShowGamePlayTags, "Show GamePlay Tags","Toggles visibility of GamePlay Tags", EUserInterfaceActionType::ToggleButton, FInputChord(EKeys::G));
-	UI_COMMAND(ShowNativePackages, "Show Native Packages","Toggles visibility of Native Packages", EUserInterfaceActionType::ToggleButton, FInputChord(EKeys::N));
+	UI_COMMAND(ShowNameReferences, "Show Name References","Toggles visibility of Name References (i.e. Gameplay Tags and Data Table Row Handles)", EUserInterfaceActionType::ToggleButton, FInputChord(EKeys::N));
+	UI_COMMAND(ShowCodePackages, "Show C++ Packages","Toggles visibility of C++ Packages", EUserInterfaceActionType::ToggleButton, FInputChord(EKeys::C));
 
 	UI_COMMAND(ShowDuplicates, "Show Duplicate References", "Toggles visibility of Duplicate References", EUserInterfaceActionType::ToggleButton, FInputChord(EKeys::U));
-	UI_COMMAND(CompactMode, "Compact Mode", "Toggles Compact Mode", EUserInterfaceActionType::ToggleButton, FInputChord(EKeys::C));
+	UI_COMMAND(CompactMode, "Compact Mode", "Toggles Compact View", EUserInterfaceActionType::ToggleButton, FInputChord(EKeys::V));
 	UI_COMMAND(FilterSearch, "Filter Search Results", "Toggles filtering of search results", EUserInterfaceActionType::ToggleButton, FInputChord());
 
 	UI_COMMAND(ShowCommentPath, "Show Asset Path", "Toggles visibility of the Asset Path shown as a comment", EUserInterfaceActionType::ToggleButton, FInputChord(EKeys::P));
diff --git a/Engine/Plugins/Editor/AssetManagerEditor/Source/AssetManagerEditor/Private/AssetManagerEditorCommands.h b/Engine/Plugins/Editor/AssetManagerEditor/Source/AssetManagerEditor/Private/AssetManagerEditorCommands.h
--- a/Engine/Plugins/Editor/AssetManagerEditor/Source/AssetManagerEditor/Private/AssetManagerEditorCommands.h
+++ b/Engine/Plugins/Editor/AssetManagerEditor/Source/AssetManagerEditor/Private/AssetManagerEditorCommands.h
@@ -73,11 +73,11 @@ public:
 	// Toggles visiblity of Management References (i.e. PrimaryAssetIDs)
 	TSharedPtr<FUICommandInfo> ShowManagementReferences;
 
-	// Toggles visiblity of GameplayTags
-	TSharedPtr<FUICommandInfo> ShowGamePlayTags;
+	// Toggles visiblity of Name References (i.e. Gameplay Tags and Data Table Row Handles)
+	TSharedPtr<FUICommandInfo> ShowNameReferences;
 
 	// Toggles visiblity of Native Packages
-	TSharedPtr<FUICommandInfo> ShowNativePackages;
+	TSharedPtr<FUICommandInfo> ShowCodePackages;
 
 	// Toggles visiblity of Duplicate References
 	TSharedPtr<FUICommandInfo> ShowDuplicates;
diff --git a/Engine/Plugins/Editor/AssetManagerEditor/Source/AssetManagerEditor/Private/ReferenceViewer/EdGraphNode_Reference.cpp b/Engine/Plugins/Editor/AssetManagerEditor/Source/AssetManagerEditor/Private/ReferenceViewer/EdGraphNode_Reference.cpp
--- a/Engine/Plugins/Editor/AssetManagerEditor/Source/AssetManagerEditor/Private/ReferenceViewer/EdGraphNode_Reference.cpp
+++ b/Engine/Plugins/Editor/AssetManagerEditor/Source/AssetManagerEditor/Private/ReferenceViewer/EdGraphNode_Reference.cpp
@@ -61,21 +61,29 @@ void UEdGraphNode_Reference::SetupReferenceNode(const FIntPoint& NodeLoc, const
 	FPrimaryAssetId PrimaryAssetID = NewIdentifiers[0].GetPrimaryAssetId();
 	if (PrimaryAssetID.IsValid())  // Management References (PrimaryAssetIDs)
 	{
-		MainAssetName = PrimaryAssetID.PrimaryAssetName.ToString();
-		AssetTypeName = PrimaryAssetID.PrimaryAssetType.ToString();
+		static FText ManagerText = LOCTEXT("ReferenceManager", "Manager");
+		MainAssetName = PrimaryAssetID.PrimaryAssetType.ToString() + TEXT(":") + PrimaryAssetID.PrimaryAssetName.ToString();
+		AssetTypeName = ManagerText.ToString();
 		bIsPackage = false;
 		bIsPrimaryAsset = true;
 	}
-	else if (First.IsValue()) // Searchable Names (GamePlay Tags)
+	else if (First.IsValue()) // Searchable Names (GamePlay Tags, Data Table Row Handle)
 	{
 		MainAssetName = First.ValueName.ToString();
 		AssetTypeName = First.ObjectName.ToString();
+		static const FName NAME_DataTable(TEXT("DataTable"));
+		static const FText InDataTableText = LOCTEXT("InDataTable", "In DataTable");
+		if (InAssetData.AssetClassPath.GetAssetName() == NAME_DataTable)
+		{
+			AssetTypeName = InDataTableText.ToString() + TEXT(" ") + AssetTypeName;
+		}
+
 		bIsPackage = false;
 	}
 	else if (First.IsPackage() && !InAssetData.IsValid()) 
 	{
 		const FString PackageNameStr = Identifiers[0].PackageName.ToString();
-		if ( PackageNameStr.StartsWith(TEXT("/Script")) )// Native Packages (/Script Code)
+		if ( PackageNameStr.StartsWith(TEXT("/Script")) )// C++ Packages (/Script Code)
 		{
 			MainAssetName = PackageNameStr.RightChop(8);
 			AssetTypeName = TEXT("Script");
diff --git a/Engine/Plugins/Editor/AssetManagerEditor/Source/AssetManagerEditor/Private/ReferenceViewer/EdGraph_ReferenceViewer.cpp b/Engine/Plugins/Editor/AssetManagerEditor/Source/AssetManagerEditor/Private/ReferenceViewer/EdGraph_ReferenceViewer.cpp
--- a/Engine/Plugins/Editor/AssetManagerEditor/Source/AssetManagerEditor/Private/ReferenceViewer/EdGraph_ReferenceViewer.cpp
+++ b/Engine/Plugins/Editor/AssetManagerEditor/Source/AssetManagerEditor/Private/ReferenceViewer/EdGraph_ReferenceViewer.cpp
@@ -371,10 +371,23 @@ void UEdGraph_ReferenceViewer::GetSortedLinks(const TArray<FAssetIdentifier>& Id
 	{
 		switch (InCategory)
 		{
-		case EDependencyCategory::Package: return 0;
-		case EDependencyCategory::Manage: return 1;
-		case EDependencyCategory::SearchableName: return 2;
-		default: check(false);  return 3;
+			case EDependencyCategory::Package: 
+			{
+				return 0;
+			}
+			case EDependencyCategory::Manage:
+			{
+				return 1;
+			}
+			case EDependencyCategory::SearchableName: 
+			{
+				return 2;
+			}
+			default: 
+			{
+				check(false);
+				return 3;
+			}
 		}
 	};
 	auto IsHard = [](EDependencyProperty Properties)
@@ -455,7 +468,7 @@ bool UEdGraph_ReferenceViewer::IsPackageIdentifierPassingFilter(const FAssetIden
 {
 	if (!InAssetIdentifier.IsValue())
 	{
-		if (!Settings->IsShowNativePackages() && InAssetIdentifier.PackageName.ToString().StartsWith(TEXT("/Script")))
+		if (!Settings->IsShowCodePackages() && InAssetIdentifier.PackageName.ToString().StartsWith(TEXT("/Script")))
 		{
 			return false;
 		}
diff --git a/Engine/Plugins/Editor/AssetManagerEditor/Source/AssetManagerEditor/Private/ReferenceViewer/ReferenceViewerSettings.cpp b/Engine/Plugins/Editor/AssetManagerEditor/Source/AssetManagerEditor/Private/ReferenceViewer/ReferenceViewerSettings.cpp
--- a/Engine/Plugins/Editor/AssetManagerEditor/Source/AssetManagerEditor/Private/ReferenceViewer/ReferenceViewerSettings.cpp
+++ b/Engine/Plugins/Editor/AssetManagerEditor/Source/AssetManagerEditor/Private/ReferenceViewer/ReferenceViewerSettings.cpp
@@ -52,9 +52,9 @@ bool UReferenceViewerSettings::IsShowSearchableNames() const
 	return bIsShowSearchableNames;
 }
 
-bool UReferenceViewerSettings::IsShowNativePackages() const
+bool UReferenceViewerSettings::IsShowCodePackages() const
 {
-	return bIsShowNativePackages;
+	return bIsShowCodePackages;
 }
 
 bool UReferenceViewerSettings::IsShowReferencers() const
@@ -127,9 +127,9 @@ void UReferenceViewerSettings::SetShowSearchableNames(bool newEnabled)
 	SaveConfig();
 }
 
-void UReferenceViewerSettings::SetShowNativePackages(bool newEnabled)
+void UReferenceViewerSettings::SetShowCodePackages(bool newEnabled)
 {
-	bIsShowNativePackages = newEnabled;
+	bIsShowCodePackages = newEnabled;
 	SaveConfig();
 }
 
diff --git a/Engine/Plugins/Editor/AssetManagerEditor/Source/AssetManagerEditor/Private/ReferenceViewer/SReferenceViewer.cpp b/Engine/Plugins/Editor/AssetManagerEditor/Source/AssetManagerEditor/Private/ReferenceViewer/SReferenceViewer.cpp
--- a/Engine/Plugins/Editor/AssetManagerEditor/Source/AssetManagerEditor/Private/ReferenceViewer/SReferenceViewer.cpp
+++ b/Engine/Plugins/Editor/AssetManagerEditor/Source/AssetManagerEditor/Private/ReferenceViewer/SReferenceViewer.cpp
@@ -113,7 +113,7 @@ void SReferenceViewer::Construct(const FArguments& InArgs)
 	bShowCollectionFilter = true;
 	bShowShowReferencesOptions = true;
 	bShowShowSearchableNames = true;
-	bShowShowNativePackages = true;
+	bShowShowCodePackages = true;
 	bShowShowFilteredPackagesOnly = true;
 	bShowCompactMode = true;
 	bDirtyResults = false;
@@ -469,7 +469,7 @@ void SReferenceViewer::SetGraphRootIdentifiers(const TArray<FAssetIdentifier>& N
 	bShowCollectionFilter = ReferenceViewerParams.bShowCollectionFilter;
 	bShowShowReferencesOptions = ReferenceViewerParams.bShowShowReferencesOptions;
 	bShowShowSearchableNames = ReferenceViewerParams.bShowShowSearchableNames;
-	bShowShowNativePackages = ReferenceViewerParams.bShowShowNativePackages;
+	bShowShowCodePackages = ReferenceViewerParams.bShowShowCodePackages;
 
 	bShowShowFilteredPackagesOnly = ReferenceViewerParams.bShowShowFilteredPackagesOnly;
 	if (ReferenceViewerParams.bShowFilteredPackagesOnly.IsSet())
@@ -1069,15 +1069,15 @@ bool SReferenceViewer::IsShowSearchableNamesChecked() const
 	return Settings->IsShowSearchableNames();
 }
 
-void SReferenceViewer::OnShowNativePackagesChanged()
+void SReferenceViewer::OnShowCodePackagesChanged()
 {
-	Settings->SetShowNativePackages(!Settings->IsShowNativePackages());
+	Settings->SetShowCodePackages(!Settings->IsShowCodePackages());
 	RebuildGraph();
 }
 
-bool SReferenceViewer::IsShowNativePackagesChecked() const
+bool SReferenceViewer::IsShowCodePackagesChecked() const
 {
-	return Settings->IsShowNativePackages();
+	return Settings->IsShowCodePackages();
 }
 
 int32 SReferenceViewer::GetSearchBreadthCount() const
@@ -1197,18 +1197,18 @@ void SReferenceViewer::RegisterActions()
 		FIsActionButtonVisible::CreateSP(this, &SReferenceViewer::GetManagementReferencesVisibility));
 
 	ReferenceViewerActions->MapAction(
-		FAssetManagerEditorCommands::Get().ShowGamePlayTags,
+		FAssetManagerEditorCommands::Get().ShowNameReferences,
 		FExecuteAction::CreateSP(this, &SReferenceViewer::OnShowSearchableNamesChanged),
 		FCanExecuteAction(),	
 		FIsActionChecked::CreateSP(this, &SReferenceViewer::IsShowSearchableNamesChecked),
 		FIsActionButtonVisible::CreateLambda([this] { return bShowShowSearchableNames; }));
 
 	ReferenceViewerActions->MapAction(
-		FAssetManagerEditorCommands::Get().ShowNativePackages,
-		FExecuteAction::CreateSP(this, &SReferenceViewer::OnShowNativePackagesChanged),
+		FAssetManagerEditorCommands::Get().ShowCodePackages,
+		FExecuteAction::CreateSP(this, &SReferenceViewer::OnShowCodePackagesChanged),
 		FCanExecuteAction(),	
-		FIsActionChecked::CreateSP(this, &SReferenceViewer::IsShowNativePackagesChecked),
-		FIsActionButtonVisible::CreateLambda([this] { return bShowShowNativePackages; }));
+		FIsActionChecked::CreateSP(this, &SReferenceViewer::IsShowCodePackagesChecked),
+		FIsActionButtonVisible::CreateLambda([this] { return bShowShowCodePackages; }));
 
 
 	ReferenceViewerActions->MapAction(
@@ -1900,8 +1900,8 @@ TSharedRef<SWidget> SReferenceViewer::GetShowMenuContent()
 
 	MenuBuilder.BeginSection("Assets", LOCTEXT("Assets", "Assets"));
 	MenuBuilder.AddMenuEntry(FAssetManagerEditorCommands::Get().ShowManagementReferences);
-	MenuBuilder.AddMenuEntry(FAssetManagerEditorCommands::Get().ShowGamePlayTags);
-	MenuBuilder.AddMenuEntry(FAssetManagerEditorCommands::Get().ShowNativePackages);
+	MenuBuilder.AddMenuEntry(FAssetManagerEditorCommands::Get().ShowNameReferences);
+	MenuBuilder.AddMenuEntry(FAssetManagerEditorCommands::Get().ShowCodePackages);
 	MenuBuilder.EndSection();
 
 	MenuBuilder.BeginSection("ViewOptions", LOCTEXT("ViewOptions", "View Options"));
diff --git a/Engine/Plugins/Editor/AssetManagerEditor/Source/AssetManagerEditor/Private/ReferenceViewer/SReferenceViewer.h b/Engine/Plugins/Editor/AssetManagerEditor/Source/AssetManagerEditor/Private/ReferenceViewer/SReferenceViewer.h
--- a/Engine/Plugins/Editor/AssetManagerEditor/Source/AssetManagerEditor/Private/ReferenceViewer/SReferenceViewer.h
+++ b/Engine/Plugins/Editor/AssetManagerEditor/Source/AssetManagerEditor/Private/ReferenceViewer/SReferenceViewer.h
@@ -150,8 +150,8 @@ private:
 	bool IsShowManagementReferencesChecked() const;
 	void OnShowSearchableNamesChanged();
 	bool IsShowSearchableNamesChecked() const;
-	void OnShowNativePackagesChanged();
-	bool IsShowNativePackagesChecked() const;
+	void OnShowCodePackagesChanged();
+	bool IsShowCodePackagesChecked() const;
 
 	int32 GetSearchBreadthCount() const;
 	void OnSearchBreadthCommitted(int32 NewValue);
@@ -242,8 +242,8 @@ private:
 	bool bShowShowReferencesOptions;
 	/** Whether to visually show to the user the option of "Show Searchable Names" */
 	bool bShowShowSearchableNames;
-	/** Whether to visually show to the user the option of "Show Native Packages" */
-	bool bShowShowNativePackages;
+	/** Whether to visually show to the user the option of "Show C++ Packages" */
+	bool bShowShowCodePackages;
 	/** Whether to visually show to the user the option of "Show Filtered Packages Only" */
 	bool bShowShowFilteredPackagesOnly;
 	/** True if our view is out of date due to asset registry changes */
diff --git a/Engine/Plugins/Editor/AssetManagerEditor/Source/AssetManagerEditor/Public/ReferenceViewer/ReferenceViewerSettings.h b/Engine/Plugins/Editor/AssetManagerEditor/Source/AssetManagerEditor/Public/ReferenceViewer/ReferenceViewerSettings.h
--- a/Engine/Plugins/Editor/AssetManagerEditor/Source/AssetManagerEditor/Public/ReferenceViewer/ReferenceViewerSettings.h
+++ b/Engine/Plugins/Editor/AssetManagerEditor/Source/AssetManagerEditor/Public/ReferenceViewer/ReferenceViewerSettings.h
@@ -77,8 +77,8 @@ public:
 	bool IsShowSearchableNames() const;
 	void SetShowSearchableNames(bool newEnabled);
 
-	bool IsShowNativePackages() const;
-	void SetShowNativePackages(bool newEnabled);
+	bool IsShowCodePackages() const;
+	void SetShowCodePackages(bool newEnabled);
 
 	bool IsShowDuplicates() const;
 	void SetShowDuplicatesEnabled(bool newEnabled);
@@ -157,7 +157,7 @@ private:
 	
 	/* Show/Hide Native Packages  */
 	UPROPERTY(config)
-	bool bIsShowNativePackages;
+	bool bIsShowCodePackages;
 	
 	/* Whether to show duplicate asset references */
 	UPROPERTY(config)
diff --git a/Engine/Source/Runtime/CoreUObject/Public/AssetRegistry/AssetData.h b/Engine/Source/Runtime/CoreUObject/Public/AssetRegistry/AssetData.h
--- a/Engine/Source/Runtime/CoreUObject/Public/AssetRegistry/AssetData.h
+++ b/Engine/Source/Runtime/CoreUObject/Public/AssetRegistry/AssetData.h
@@ -822,6 +822,7 @@ PRAGMA_ENABLE_DEPRECATION_WARNINGS
  * Helper struct for FAssetIdentifier (e.g., for the FOnViewAssetIdentifiersInReferenceViewer delegate and Reference Viewer functions).
  */
 #if WITH_EDITORONLY_DATA
+PRAGMA_DISABLE_DEPRECATION_WARNINGS
 struct FReferenceViewerParams
 {
 	FReferenceViewerParams()
@@ -836,6 +837,7 @@ struct FReferenceViewerParams
 		, bShowShowReferencesOptions(true)
 		, bShowShowSearchableNames(true)
 		, bShowShowNativePackages(true)
+		, bShowShowCodePackages(true)
 		, bShowShowFilteredPackagesOnly(true)
 		, bShowCompactMode(true)
 	{}
@@ -869,12 +871,16 @@ struct FReferenceViewerParams
 	/** Whether to visually show to the user the option of "Show Searchable Names" */
 	bool bShowShowSearchableNames;
 	/** Whether to visually show to the user the option of "Show Native Packages" */
+	UE_DEPRECATED(5.1, "bShowShowNativePackages is deprecated, please use bShowShowCodePackages instead.")
 	bool bShowShowNativePackages;
+	/** Whether to visually show to the user the option of "Show C++ Packages" */
+	bool bShowShowCodePackages;
 	/** Whether to visually show to the user the option of "Show Filtered Packages Only" */
 	bool bShowShowFilteredPackagesOnly;
 	/** Whether to visually show to the user the option of "Compact Mode" */
 	bool bShowCompactMode;
 };
+PRAGMA_ENABLE_DEPRECATION_WARNINGS
 #endif // WITH_EDITORONLY_DATA
 
 /** A structure defining a thing that can be reference by something else in the asset registry. Represents either a package of a primary asset id */
