Reference Viewer: Fix filters not wrapping properly

diff --git a/Engine/Plugins/Editor/AssetManagerEditor/Source/AssetManagerEditor/Private/ReferenceViewer/SReferenceViewer.cpp b/Engine/Plugins/Editor/AssetManagerEditor/Source/AssetManagerEditor/Private/ReferenceViewer/SReferenceViewer.cpp
--- a/Engine/Plugins/Editor/AssetManagerEditor/Source/AssetManagerEditor/Private/ReferenceViewer/SReferenceViewer.cpp
+++ b/Engine/Plugins/Editor/AssetManagerEditor/Source/AssetManagerEditor/Private/ReferenceViewer/SReferenceViewer.cpp
@@ -194,7 +194,7 @@ void SReferenceViewer::Construct(const FArguments& InArgs)
 
 			+SOverlay::Slot()
 			.VAlign(VAlign_Top)
-			.HAlign(HAlign_Left)
+			.HAlign(HAlign_Fill)
 			.Padding(8)
 			[
 				SNew(SHorizontalBox)
