Reference Viewer: Fix an issue with Autofilters hiding some assets, even though they have a parent class which has a filter

diff --git a/Engine/Plugins/Editor/AssetManagerEditor/Source/AssetManagerEditor/Private/ReferenceViewer/SReferenceViewer.cpp b/Engine/Plugins/Editor/AssetManagerEditor/Source/AssetManagerEditor/Private/ReferenceViewer/SReferenceViewer.cpp
--- a/Engine/Plugins/Editor/AssetManagerEditor/Source/AssetManagerEditor/Private/ReferenceViewer/SReferenceViewer.cpp
+++ b/Engine/Plugins/Editor/AssetManagerEditor/Source/AssetManagerEditor/Private/ReferenceViewer/SReferenceViewer.cpp
@@ -775,6 +775,23 @@ void SReferenceViewer::OnUpdateFilterBar()
 				{
 					FilterWidget->SetAssetTypeFilterCheckState(AssetClassPath, ECheckBoxState::Checked);
 				}
+				// If the current AssetClassPath does not have a filter in the filter bar, we walk through its ancestor classes to see if any of those have a filter
+				else
+				{
+					TArray<FTopLevelAssetPath> AncestorClassNames;
+					
+					FAssetRegistryModule& AssetRegistryModule = FModuleManager::GetModuleChecked< FAssetRegistryModule >( TEXT("AssetRegistry") );
+					AssetRegistryModule.Get().GetAncestorClassNames(AssetClassPath, AncestorClassNames);
+					
+					for (const FTopLevelAssetPath& AssetClassAncestor : AncestorClassNames)
+					{
+						if (FilterWidget->DoesAssetTypeFilterExist(AssetClassAncestor))
+						{
+							FilterWidget->SetAssetTypeFilterCheckState(AssetClassAncestor, ECheckBoxState::Checked);
+							break;
+						}
+					}
+				}
 			}
 
 			GraphObj->SetCurrentFilterCollection(FilterWidget->GetAllActiveFilters());
