Reference Viewer: Improves UI Responsiveness for Referencer and Dependency Depth Changes on large scenes

diff --git a/Engine/Plugins/Editor/AssetManagerEditor/Source/AssetManagerEditor/Private/ReferenceViewer/SReferenceViewer.cpp b/Engine/Plugins/Editor/AssetManagerEditor/Source/AssetManagerEditor/Private/ReferenceViewer/SReferenceViewer.cpp
--- a/Engine/Plugins/Editor/AssetManagerEditor/Source/AssetManagerEditor/Private/ReferenceViewer/SReferenceViewer.cpp
+++ b/Engine/Plugins/Editor/AssetManagerEditor/Source/AssetManagerEditor/Private/ReferenceViewer/SReferenceViewer.cpp
@@ -71,6 +71,7 @@ SReferenceViewer::~SReferenceViewer()
 void SReferenceViewer::Construct(const FArguments& InArgs)
 {
 	bRebuildingFilters = false;
+	bNeedsGraphRebuild = false;
 	Settings = GetMutableDefault<UReferenceViewerSettings>();
 
 	// Create an action list and register commands
@@ -252,7 +253,7 @@ void SReferenceViewer::Construct(const FArguments& InArgs)
 								.OnValueCommitted_Lambda([this] (int32 NewValue, ETextCommit::Type CommitType) { FSlateApplication::Get().SetKeyboardFocus(GraphEditorPtr, EFocusCause::SetDirectly); } )
 								.MinValue(0)
 								.MaxValue(50)
-								.MaxSliderValue(12)
+								.MaxSliderValue(10)
 							]
 						]
 					]
@@ -289,7 +290,7 @@ void SReferenceViewer::Construct(const FArguments& InArgs)
 								.OnValueCommitted_Lambda([this] (int32 NewValue, ETextCommit::Type CommitType) { FSlateApplication::Get().SetKeyboardFocus(GraphEditorPtr, EFocusCause::SetDirectly); } )
 								.MinValue(0)
 								.MaxValue(50)
-								.MaxSliderValue(12)
+								.MaxSliderValue(10)
 							]
 						]
 					]
@@ -439,6 +440,17 @@ void SReferenceViewer::Construct(const FArguments& InArgs)
 	];
 
 	UpdateCollectionsComboList();
+
+	SetCanTick(true);
+}
+
+void SReferenceViewer::Tick( const FGeometry& AllottedGeometry, const double InCurrentTime, const float InDeltaTime )
+{
+	if (bNeedsGraphRebuild)
+	{
+		bNeedsGraphRebuild = false;
+		RebuildGraph();
+	}
 }
 
 FReply SReferenceViewer::OnKeyDown(const FGeometry& MyGeometry, const FKeyEvent& InKeyEvent)
@@ -830,14 +842,24 @@ int32 SReferenceViewer::GetSearchReferencerDepthCount() const
 
 void SReferenceViewer::OnSearchDependencyDepthCommitted(int32 NewValue)
 {
-	Settings->SetSearchDependencyDepthLimit(NewValue);
-	RebuildGraph();
+	if (NewValue != Settings->GetSearchDependencyDepthLimit())
+	{
+		Settings->SetSearchDependencyDepthLimit(NewValue);
+
+		// Push the rebuild to Tick() in order to keep the slider as responsive as possible
+		bNeedsGraphRebuild = true;
+	}
 }
 
 void SReferenceViewer::OnSearchReferencerDepthCommitted(int32 NewValue)
 {
-	Settings->SetSearchReferencerDepthLimit(NewValue);
-	RebuildGraph();
+	if (NewValue != Settings->GetSearchReferencerDepthLimit())
+	{
+		Settings->SetSearchReferencerDepthLimit(NewValue);
+
+		// Push the rebuild to Tick() in order to keep the slider as responsive as possible
+		bNeedsGraphRebuild = true;
+	}
 }
 
 void SReferenceViewer::OnSearchBreadthEnabledChanged( ECheckBoxState NewState )
diff --git a/Engine/Plugins/Editor/AssetManagerEditor/Source/AssetManagerEditor/Private/ReferenceViewer/SReferenceViewer.h b/Engine/Plugins/Editor/AssetManagerEditor/Source/AssetManagerEditor/Private/ReferenceViewer/SReferenceViewer.h
--- a/Engine/Plugins/Editor/AssetManagerEditor/Source/AssetManagerEditor/Private/ReferenceViewer/SReferenceViewer.h
+++ b/Engine/Plugins/Editor/AssetManagerEditor/Source/AssetManagerEditor/Private/ReferenceViewer/SReferenceViewer.h
@@ -52,6 +52,7 @@ public:
 	/**SWidget interface **/
 	virtual bool SupportsKeyboardFocus() const override { return true; }
 	virtual FReply OnKeyDown( const FGeometry& MyGeometry, const FKeyEvent& InKeyEvent ) override;
+	virtual void Tick( const FGeometry& AllottedGeometry, const double InCurrentTime, const float InDeltaTime ) override;
 
 private:
 
@@ -254,6 +255,8 @@ private:
 	/** A recursion check so as to avoid the rebuild of the graph if we are currently rebuilding the filters */
 	bool bRebuildingFilters;
 
+	bool bNeedsGraphRebuild;
+
 	/** Handle to know if dirty */
 	FDelegateHandle AssetRefreshHandle;
 };
