Reference Viewer:  Further improves the interaction with the referencer & dependency depth sliders

diff --git a/Engine/Plugins/Editor/AssetManagerEditor/Source/AssetManagerEditor/Private/ReferenceViewer/ReferenceViewerSettings.cpp b/Engine/Plugins/Editor/AssetManagerEditor/Source/AssetManagerEditor/Private/ReferenceViewer/ReferenceViewerSettings.cpp
--- a/Engine/Plugins/Editor/AssetManagerEditor/Source/AssetManagerEditor/Private/ReferenceViewer/ReferenceViewerSettings.cpp
+++ b/Engine/Plugins/Editor/AssetManagerEditor/Source/AssetManagerEditor/Private/ReferenceViewer/ReferenceViewerSettings.cpp
@@ -150,10 +150,13 @@ int32 UReferenceViewerSettings::GetSearchDependencyDepthLimit() const
 	return MaxSearchDependencyDepth;
 }
 
-void UReferenceViewerSettings::SetSearchDependencyDepthLimit(int32 NewDepthLimit)
+void UReferenceViewerSettings::SetSearchDependencyDepthLimit(int32 NewDepthLimit, bool bSaveConfig)
 {
 	MaxSearchDependencyDepth = FMath::Max(NewDepthLimit, 0);
-	SaveConfig();
+	if (bSaveConfig)
+	{
+		SaveConfig();
+	}
 }
 
 int32 UReferenceViewerSettings::GetSearchReferencerDepthLimit() const
@@ -161,10 +164,13 @@ int32 UReferenceViewerSettings::GetSearchReferencerDepthLimit() const
 	return MaxSearchReferencerDepth;
 }
 
-void UReferenceViewerSettings::SetSearchReferencerDepthLimit(int32 NewDepthLimit)
+void UReferenceViewerSettings::SetSearchReferencerDepthLimit(int32 NewDepthLimit, bool bSaveConfig)
 {
 	MaxSearchReferencerDepth = FMath::Max(NewDepthLimit, 0);
-	SaveConfig();
+	if (bSaveConfig)
+	{
+		SaveConfig();
+	}
 }
 
 int32 UReferenceViewerSettings::GetSearchBreadthLimit() const
diff --git a/Engine/Plugins/Editor/AssetManagerEditor/Source/AssetManagerEditor/Private/ReferenceViewer/SReferenceViewer.cpp b/Engine/Plugins/Editor/AssetManagerEditor/Source/AssetManagerEditor/Private/ReferenceViewer/SReferenceViewer.cpp
--- a/Engine/Plugins/Editor/AssetManagerEditor/Source/AssetManagerEditor/Private/ReferenceViewer/SReferenceViewer.cpp
+++ b/Engine/Plugins/Editor/AssetManagerEditor/Source/AssetManagerEditor/Private/ReferenceViewer/SReferenceViewer.cpp
@@ -193,6 +193,16 @@ void SReferenceViewer::Construct(const FArguments& InArgs)
 				GraphEditorPtr.ToSharedRef()
 			]
 
+			+SOverlay::Slot()
+			.HAlign(HAlign_Fill)
+			.VAlign(VAlign_Fill)
+			[
+				SNew(SImage)
+				.Image(FAppStyle::GetBrush("Brushes.Recessed"))
+				.ColorAndOpacity_Lambda( [this] () { return bNeedsGraphRebuild ? FLinearColor(1.0, 1.0, 1.0, 0.25) : FLinearColor::Transparent; } )
+				.Visibility(EVisibility::HitTestInvisible)
+			]
+
 			+SOverlay::Slot()
 			.VAlign(VAlign_Top)
 			.HAlign(HAlign_Fill)
@@ -249,8 +259,32 @@ void SReferenceViewer::Construct(const FArguments& InArgs)
 							[
 								SAssignNew(ReferencerCountBox, SSpinBox<int32>)
 								.Value(this, &SReferenceViewer::GetSearchReferencerDepthCount)
-								.OnValueChanged(this, &SReferenceViewer::OnSearchReferencerDepthCommitted)
-								.OnValueCommitted_Lambda([this] (int32 NewValue, ETextCommit::Type CommitType) { FSlateApplication::Get().SetKeyboardFocus(GraphEditorPtr, EFocusCause::SetDirectly); } )
+								.OnValueChanged_Lambda([this] (int32 NewValue)
+									{
+										if (NewValue != Settings->GetSearchReferencerDepthLimit())
+										{
+											Settings->SetSearchReferencerDepthLimit(NewValue, false);
+											bNeedsGraphRebuild = true;
+
+											SliderDelayLastMovedTime = FSlateApplication::Get().GetCurrentTime();
+										}
+									}
+								)
+								.OnValueCommitted_Lambda([this] (int32 NewValue, ETextCommit::Type CommitType) 
+									{ 
+										FSlateApplication::Get().SetKeyboardFocus(GraphEditorPtr, EFocusCause::SetDirectly); 
+
+										if (NewValue != Settings->GetSearchReferencerDepthLimit() || bNeedsGraphRebuild)
+										{
+											Settings->SetSearchReferencerDepthLimit(NewValue, false); 
+											bNeedsGraphRebuild = false;
+											RebuildGraph();
+										}
+
+										// Always save the config since we explicitly did not save during slider movement to preserve interactivity
+										Settings->SaveConfig();
+									} 
+								)
 								.MinValue(0)
 								.MaxValue(50)
 								.MaxSliderValue(10)
@@ -286,8 +320,32 @@ void SReferenceViewer::Construct(const FArguments& InArgs)
 							[
 								SAssignNew(DependencyCountBox, SSpinBox<int32>)
 								.Value(this, &SReferenceViewer::GetSearchDependencyDepthCount)
-								.OnValueChanged(this, &SReferenceViewer::OnSearchDependencyDepthCommitted)
-								.OnValueCommitted_Lambda([this] (int32 NewValue, ETextCommit::Type CommitType) { FSlateApplication::Get().SetKeyboardFocus(GraphEditorPtr, EFocusCause::SetDirectly); } )
+								.OnValueChanged_Lambda([this] (int32 NewValue)
+									{	
+										if (NewValue != Settings->GetSearchDependencyDepthLimit())
+										{
+											Settings->SetSearchDependencyDepthLimit(NewValue, false);
+											bNeedsGraphRebuild = true;
+
+											SliderDelayLastMovedTime = FSlateApplication::Get().GetCurrentTime();
+										}
+									}
+								)
+								.OnValueCommitted_Lambda([this] (int32 NewValue, ETextCommit::Type CommitType) 
+									{ 
+										FSlateApplication::Get().SetKeyboardFocus(GraphEditorPtr, EFocusCause::SetDirectly); 
+
+										if (NewValue != Settings->GetSearchDependencyDepthLimit() || bNeedsGraphRebuild)
+										{
+											Settings->SetSearchDependencyDepthLimit(NewValue, false);
+											bNeedsGraphRebuild = false;
+											RebuildGraph();
+										}
+
+										// Always save the config since we explicitly did not save during slider movement to preserve interactivity
+										Settings->SaveConfig();
+									}
+								)
 								.MinValue(0)
 								.MaxValue(50)
 								.MaxSliderValue(10)
@@ -446,7 +504,7 @@ void SReferenceViewer::Construct(const FArguments& InArgs)
 
 void SReferenceViewer::Tick( const FGeometry& AllottedGeometry, const double InCurrentTime, const float InDeltaTime )
 {
-	if (bNeedsGraphRebuild)
+	if (bNeedsGraphRebuild && (InCurrentTime - SliderDelayLastMovedTime > GraphRebuildSliderDelay))
 	{
 		bNeedsGraphRebuild = false;
 		RebuildGraph();
@@ -845,9 +903,7 @@ void SReferenceViewer::OnSearchDependencyDepthCommitted(int32 NewValue)
 	if (NewValue != Settings->GetSearchDependencyDepthLimit())
 	{
 		Settings->SetSearchDependencyDepthLimit(NewValue);
-
-		// Push the rebuild to Tick() in order to keep the slider as responsive as possible
-		bNeedsGraphRebuild = true;
+		RebuildGraph();
 	}
 }
 
@@ -856,9 +912,7 @@ void SReferenceViewer::OnSearchReferencerDepthCommitted(int32 NewValue)
 	if (NewValue != Settings->GetSearchReferencerDepthLimit())
 	{
 		Settings->SetSearchReferencerDepthLimit(NewValue);
-
-		// Push the rebuild to Tick() in order to keep the slider as responsive as possible
-		bNeedsGraphRebuild = true;
+		RebuildGraph();
 	}
 }
 
diff --git a/Engine/Plugins/Editor/AssetManagerEditor/Source/AssetManagerEditor/Private/ReferenceViewer/SReferenceViewer.h b/Engine/Plugins/Editor/AssetManagerEditor/Source/AssetManagerEditor/Private/ReferenceViewer/SReferenceViewer.h
--- a/Engine/Plugins/Editor/AssetManagerEditor/Source/AssetManagerEditor/Private/ReferenceViewer/SReferenceViewer.h
+++ b/Engine/Plugins/Editor/AssetManagerEditor/Source/AssetManagerEditor/Private/ReferenceViewer/SReferenceViewer.h
@@ -255,7 +255,10 @@ private:
 	/** A recursion check so as to avoid the rebuild of the graph if we are currently rebuilding the filters */
 	bool bRebuildingFilters;
 
+	/** Used to delay graph rebuilding during spinbox slider interaction */
 	bool bNeedsGraphRebuild;
+	double SliderDelayLastMovedTime = 0.0;
+	double GraphRebuildSliderDelay = 0.25;
 
 	/** Handle to know if dirty */
 	FDelegateHandle AssetRefreshHandle;
diff --git a/Engine/Plugins/Editor/AssetManagerEditor/Source/AssetManagerEditor/Public/ReferenceViewer/ReferenceViewerSettings.h b/Engine/Plugins/Editor/AssetManagerEditor/Source/AssetManagerEditor/Public/ReferenceViewer/ReferenceViewerSettings.h
--- a/Engine/Plugins/Editor/AssetManagerEditor/Source/AssetManagerEditor/Public/ReferenceViewer/ReferenceViewerSettings.h
+++ b/Engine/Plugins/Editor/AssetManagerEditor/Source/AssetManagerEditor/Public/ReferenceViewer/ReferenceViewerSettings.h
@@ -45,13 +45,13 @@ public:
 	void SetShowReferencers(const bool bShouldShowReferencers);
 
 	int32 GetSearchReferencerDepthLimit() const;
-	void SetSearchReferencerDepthLimit(int32 NewDepthLimit);
+	void SetSearchReferencerDepthLimit(int32 NewDepthLimit, bool bSaveConfig = true);
 
 	bool IsShowDependencies() const;
 	void SetShowDependencies(const bool bShouldShowDependencies);
 
 	int32 GetSearchDependencyDepthLimit() const;
-	void SetSearchDependencyDepthLimit(int32 NewDepthLimit);
+	void SetSearchDependencyDepthLimit(int32 NewDepthLimit, bool bSaveConfig = true);
 
 	bool IsSearchBreadthLimited() const;
 	void SetSearchBreadthLimitEnabled(bool newEnabled);
