Reference Viewer: Improves handling when users type paths into the File Path field.

diff --git a/Engine/Plugins/Editor/AssetManagerEditor/Source/AssetManagerEditor/Private/ReferenceViewer/SReferenceViewer.cpp b/Engine/Plugins/Editor/AssetManagerEditor/Source/AssetManagerEditor/Private/ReferenceViewer/SReferenceViewer.cpp
--- a/Engine/Plugins/Editor/AssetManagerEditor/Source/AssetManagerEditor/Private/ReferenceViewer/SReferenceViewer.cpp
+++ b/Engine/Plugins/Editor/AssetManagerEditor/Source/AssetManagerEditor/Private/ReferenceViewer/SReferenceViewer.cpp
@@ -569,6 +569,8 @@ void SReferenceViewer::SetGraphRootIdentifiers(const TArray<FAssetIdentifier>& N
 
 	// Set the initial history data
 	HistoryManager.AddHistoryData();
+
+	TemporaryPathBeingEdited = NewGraphRootIdentifiers.Num() > 0 ? FText() : FText(LOCTEXT("NoAssetsFound", "No Assets Found"));
 }
 
 EActiveTimerReturnType SReferenceViewer::TriggerZoomToFit(double InCurrentTime, float InDeltaTime)
@@ -784,15 +786,32 @@ FText SReferenceViewer::GetStatusText() const
 
 void SReferenceViewer::OnAddressBarTextCommitted(const FText& NewText, ETextCommit::Type CommitInfo)
 {
+	TArray<FAssetIdentifier> NewPaths;
 	if (CommitInfo == ETextCommit::OnEnter)
 	{
-		TArray<FAssetIdentifier> NewPaths;
-		NewPaths.Add(FAssetIdentifier::FromString(NewText.ToString()));
+		TArray<FAssetData> SelectedAssets;
+
+		FAssetIdentifier NewPath = FAssetIdentifier::FromString(NewText.ToString());
+
+		FAssetRegistryModule& AssetRegistryModule = FModuleManager::GetModuleChecked< FAssetRegistryModule >( TEXT("AssetRegistry") );
+		IAssetRegistry* AssetRegistry = &AssetRegistryModule.Get();
+
+		AssetRegistry->GetAssetsByPackageName(FName(*NewText.ToString()), SelectedAssets);
+		if (SelectedAssets.Num() > 0)
+		{
+			NewPaths.Add(NewPath);
+		}
+
+		else if (AssetRegistry->GetAssetsByPath(FName(*NewText.ToString()), SelectedAssets, true))
+		{
+			for (const FAssetData& AssetData : SelectedAssets)
+			{
+				NewPaths.AddUnique(FAssetIdentifier(AssetData.PackageName));
+			}
+		}	
 
 		SetGraphRootIdentifiers(NewPaths);
 	}
-
-	TemporaryPathBeingEdited = FText();
 }
 
 void SReferenceViewer::OnAddressBarTextChanged(const FText& NewText)
@@ -811,6 +830,8 @@ void SReferenceViewer::OnApplyHistoryData(const FReferenceViewerHistoryData& His
 		{
 			GraphEditorPtr->SetNodeSelection(NewRootNode, true);
 		}
+
+		TemporaryPathBeingEdited = FText();
 	}
 }
 
