Reference Viewer: Initial Find Path Tool

diff --git a/Engine/Plugins/Editor/AssetManagerEditor/Source/AssetManagerEditor/Private/AssetManagerEditorCommands.cpp b/Engine/Plugins/Editor/AssetManagerEditor/Source/AssetManagerEditor/Private/AssetManagerEditorCommands.cpp
--- a/Engine/Plugins/Editor/AssetManagerEditor/Source/AssetManagerEditor/Private/AssetManagerEditorCommands.cpp
+++ b/Engine/Plugins/Editor/AssetManagerEditor/Source/AssetManagerEditor/Private/AssetManagerEditorCommands.cpp
@@ -53,6 +53,7 @@ void FAssetManagerEditorCommands::RegisterCommands()
 	UI_COMMAND(ShowCommentPath, "Show Asset Path", "Toggles visibility of the Asset Path shown as a comment", EUserInterfaceActionType::ToggleButton, FInputChord(EKeys::P));
 	UI_COMMAND(Filters, "Toggle the AssetType Filters", "Toggles the AssetType Filters", EUserInterfaceActionType::ToggleButton, FInputChord(EKeys::A));
 	UI_COMMAND(AutoFilters, "Auto Populate Filters", "Toggles autopopulating of AssetType Filters", EUserInterfaceActionType::ToggleButton, FInputChord());
+	UI_COMMAND(FindPath, "Enable Find Path Tool", "Toggle the Find Path Tool", EUserInterfaceActionType::ToggleButton, FInputChord());
 
 	UI_COMMAND(CopyPaths, "Copy the Asset Path", "Copies the Asset Path to the Clipboard", EUserInterfaceActionType::Button, FInputChord(EModifierKey::Control, EKeys::C));
 
diff --git a/Engine/Plugins/Editor/AssetManagerEditor/Source/AssetManagerEditor/Private/AssetManagerEditorCommands.h b/Engine/Plugins/Editor/AssetManagerEditor/Source/AssetManagerEditor/Private/AssetManagerEditorCommands.h
--- a/Engine/Plugins/Editor/AssetManagerEditor/Source/AssetManagerEditor/Private/AssetManagerEditorCommands.h
+++ b/Engine/Plugins/Editor/AssetManagerEditor/Source/AssetManagerEditor/Private/AssetManagerEditorCommands.h
@@ -100,6 +100,9 @@ public:
 	// Toggles whether search results are filtered or just selected 
 	TSharedPtr<FUICommandInfo> FilterSearch;
 
+	// Toggles the Find Path Mode
+	TSharedPtr<FUICommandInfo> FindPath;
+
 	// Copies the list of objects that the selected asset references
 	TSharedPtr<FUICommandInfo> CopyReferencedObjects;
 
diff --git a/Engine/Plugins/Editor/AssetManagerEditor/Source/AssetManagerEditor/Private/ReferenceViewer/EdGraph_ReferenceViewer.cpp b/Engine/Plugins/Editor/AssetManagerEditor/Source/AssetManagerEditor/Private/ReferenceViewer/EdGraph_ReferenceViewer.cpp
--- a/Engine/Plugins/Editor/AssetManagerEditor/Source/AssetManagerEditor/Private/ReferenceViewer/EdGraph_ReferenceViewer.cpp
+++ b/Engine/Plugins/Editor/AssetManagerEditor/Source/AssetManagerEditor/Private/ReferenceViewer/EdGraph_ReferenceViewer.cpp
@@ -119,8 +119,16 @@ bool UEdGraph_ReferenceViewer::GetSelectedAssetsForMenuExtender(const class UEdG
 UEdGraphNode_Reference* UEdGraph_ReferenceViewer::RebuildGraph()
 {
 	RemoveAllNodes();
-	UEdGraphNode_Reference* NewRootNode = ConstructNodes(CurrentGraphRootIdentifiers, CurrentGraphRootOrigin);
-	NotifyGraphChanged();
+	UEdGraphNode_Reference* NewRootNode = nullptr;
+
+	if (Settings->GetFindPathEnabled())
+	{
+		NewRootNode = FindPath(CurrentGraphRootIdentifiers[0], TargetIdentifier);
+	}
+	else
+	{
+		NewRootNode = ConstructNodes(CurrentGraphRootIdentifiers, CurrentGraphRootOrigin);
+	}
 
 	return NewRootNode;
 }
@@ -170,7 +178,7 @@ FAssetManagerDependencyQuery UEdGraph_ReferenceViewer::GetReferenceSearchFlags(b
 
 UEdGraphNode_Reference* UEdGraph_ReferenceViewer::ConstructNodes(const TArray<FAssetIdentifier>& GraphRootIdentifiers, const FIntPoint& GraphRootOrigin )
 {
-	UEdGraphNode_Reference* RootNode = NULL;
+	UEdGraphNode_Reference* RootNode = nullptr;
 	if (GraphRootIdentifiers.Num() > 0)
 	{
 		// It both were false, nothing (other than the GraphRootIdentifiers) would be displayed
@@ -196,12 +204,18 @@ UEdGraphNode_Reference* UEdGraph_ReferenceViewer::ConstructNodes(const TArray<FA
 		TMap<FAssetIdentifier, FReferenceNodeInfo> NewReferenceNodeInfos;
 		FReferenceNodeInfo& RootNodeInfo = NewReferenceNodeInfos.FindOrAdd( GraphRootIdentifiers[0], FReferenceNodeInfo(GraphRootIdentifiers[0], true));
 		RootNodeInfo.Parents.Emplace(FAssetIdentifier(NAME_None));
-		RecursivelyPopulateNodeInfos(true, GraphRootIdentifiers, NewReferenceNodeInfos, 0, Settings->GetSearchReferencerDepthLimit());
+		if (!Settings->GetFindPathEnabled())
+		{
+			RecursivelyPopulateNodeInfos(true, GraphRootIdentifiers, NewReferenceNodeInfos, 0, Settings->GetSearchReferencerDepthLimit());
+		}
 
 		TMap<FAssetIdentifier, FReferenceNodeInfo> NewDependencyNodeInfos;
 		FReferenceNodeInfo& DRootNodeInfo = NewDependencyNodeInfos.FindOrAdd( GraphRootIdentifiers[0], FReferenceNodeInfo(GraphRootIdentifiers[0], false));
 		DRootNodeInfo.Parents.Emplace(FAssetIdentifier(NAME_None));
-		RecursivelyPopulateNodeInfos(false, GraphRootIdentifiers, NewDependencyNodeInfos, 0, Settings->GetSearchDependencyDepthLimit());
+		if (!Settings->GetFindPathEnabled())
+		{
+			RecursivelyPopulateNodeInfos(false, GraphRootIdentifiers, NewDependencyNodeInfos, 0, Settings->GetSearchDependencyDepthLimit());
+		}
 
 		TSet<FName> AllPackageNames;
 		auto AddPackage = [](const FAssetIdentifier& AssetId, TSet<FName>& PackageNames)
@@ -231,7 +245,7 @@ UEdGraphNode_Reference* UEdGraph_ReferenceViewer::ConstructNodes(const TArray<FA
 
 		// Store the AssetData in the NodeInfos
 		TMap<FName, FAssetData> PackagesToAssetDataMap;
-		GatherAssetData(AllPackageNames, PackagesToAssetDataMap);
+		UE::AssetRegistry::GetAssetForPackages(AllPackageNames.Array(), PackagesToAssetDataMap);
 
 		// Store the AssetData in the NodeInfos and collect Asset Type UClasses to populate the filters
 		TSet<FTopLevelAssetPath> AllClasses;
@@ -269,11 +283,80 @@ UEdGraphNode_Reference* UEdGraph_ReferenceViewer::ConstructNodes(const TArray<FA
 	return RefilterGraph();
 }
 
+UEdGraphNode_Reference* UEdGraph_ReferenceViewer::FindPath(const FAssetIdentifier& RootId, const FAssetIdentifier& TargetId)
+{
+
+	TargetIdentifier = TargetId;
+
+	RemoveAllNodes();
+
+	// Check for the target in the dependencies
+	TMap<FAssetIdentifier, FReferenceNodeInfo> NewNodeInfos;
+	TSet<FAssetIdentifier> Visited;
+	FReferenceNodeInfo& RootNodeInfo = NewNodeInfos.FindOrAdd( RootId, FReferenceNodeInfo(RootId, false));
+	if (TargetId.IsValid())
+	{
+		FindPath_Recursive(false, RootId, TargetId, NewNodeInfos, Visited);
+	}
+	GatherAssetData(NewNodeInfos);
+	DependencyNodeInfos = NewNodeInfos;
+
+	// Check for the target in the references
+	Visited.Empty();
+	TMap<FAssetIdentifier, FReferenceNodeInfo> NewRefNodeInfos;
+	FReferenceNodeInfo& RootRefNodeInfo = NewRefNodeInfos.FindOrAdd( RootId, FReferenceNodeInfo(RootId, true));
+	if (TargetId.IsValid())
+	{
+		FindPath_Recursive(true, RootId, TargetId, NewRefNodeInfos, Visited);
+	}
+	GatherAssetData(NewRefNodeInfos);
+	ReferencerNodeInfos = NewRefNodeInfos;
+
+	UEdGraphNode_Reference* NewRootNode = RefilterGraph();
+
+	NotifyGraphChanged();
+
+	return NewRootNode;
+}
+
+bool UEdGraph_ReferenceViewer::FindPath_Recursive(bool bInReferencers, const FAssetIdentifier& InAssetId, const FAssetIdentifier& TargetId, TMap<FAssetIdentifier, FReferenceNodeInfo>& InNodeInfos, TSet<FAssetIdentifier>& Visited )
+{
+	bool bFound = false;
+
+	if (InAssetId == TargetId)
+	{
+		FReferenceNodeInfo& NewNodeInfo = InNodeInfos.FindOrAdd(InAssetId, FReferenceNodeInfo(InAssetId, bInReferencers));
+		bFound = true;
+	}
+
+	// check if any decedents are the target and if any are found, add a node info for this asset as well 
+	else 
+	{
+		Visited.Add(InAssetId);
+		TMap<FAssetIdentifier, EDependencyPinCategory> ReferenceLinks;
+		GetSortedLinks({InAssetId}, bInReferencers, GetReferenceSearchFlags(false), ReferenceLinks);
+
+		for (const TPair<FAssetIdentifier, EDependencyPinCategory>& Pair : ReferenceLinks)
+		{
+			FAssetIdentifier ChildId = Pair.Key;
+			if (!Visited.Contains(ChildId) && FindPath_Recursive(bInReferencers, ChildId, TargetId, InNodeInfos, Visited))
+			{
+				FReferenceNodeInfo& NewNodeInfo = InNodeInfos.FindOrAdd(InAssetId, FReferenceNodeInfo(InAssetId, bInReferencers));
+
+				InNodeInfos[ChildId].Parents.AddUnique(InAssetId);
+				InNodeInfos[InAssetId].Children.AddUnique(Pair);
+				bFound = true;
+			}
+		}
+	}
+
+	return bFound;
+}
 
 UEdGraphNode_Reference* UEdGraph_ReferenceViewer::RefilterGraph()
 {
 	RemoveAllNodes();
-	UEdGraphNode_Reference* RootNode = NULL;
+	UEdGraphNode_Reference* RootNode = nullptr;
 
 	bBreadthLimitReached = false;
 	if (CurrentGraphRootIdentifiers.Num() > 0 && (!ReferencerNodeInfos.IsEmpty() || !DependencyNodeInfos.IsEmpty()))
@@ -299,7 +382,6 @@ UEdGraphNode_Reference* UEdGraph_ReferenceViewer::RefilterGraph()
 			RecursivelyFilterNodeInfos(FirstGraphRootIdentifier, DependencyNodeInfos, 0, Settings->GetSearchDependencyDepthLimit());
 			RecursivelyCreateNodes(false, FirstGraphRootIdentifier, CurrentGraphRootOrigin, FirstGraphRootIdentifier, RootNode, DependencyNodeInfos, 0, Settings->GetSearchDependencyDepthLimit(), /*bIsRoot*/ true);
 		}
-
 	}
 
 	NotifyGraphChanged();
@@ -308,13 +390,13 @@ UEdGraphNode_Reference* UEdGraph_ReferenceViewer::RefilterGraph()
 
 void UEdGraph_ReferenceViewer::RecursivelyFilterNodeInfos(const FAssetIdentifier& InAssetId, TMap<FAssetIdentifier, FReferenceNodeInfo>& InNodeInfos, int32 InCurrentDepth, int32 InMaxDepth)
 {
-	// Filters and Reprovisions the NodeInfo counts 
+	// Filters and Re-provisions the NodeInfo counts 
 	int32 NewProvisionSize = 0;
 
 	int32 Breadth = 0;
 
 	InNodeInfos[InAssetId].OverflowCount = 0;
-	if (InMaxDepth > 0 && InCurrentDepth < InMaxDepth)
+	if (!ExceedsMaxSearchDepth(InCurrentDepth, InMaxDepth))
 	{
 		for (const TPair<FAssetIdentifier, EDependencyPinCategory>& Pair : InNodeInfos[InAssetId].Children)
 		{
@@ -326,6 +408,10 @@ void UEdGraph_ReferenceViewer::RecursivelyFilterNodeInfos(const FAssetIdentifier
 				RecursivelyFilterNodeInfos(ChildId, InNodeInfos, InCurrentDepth + 1, InMaxDepth);
 				ChildProvSize = InNodeInfos[ChildId].ProvisionSize(InAssetId);
 			}
+			else if (Settings->GetFindPathEnabled())
+			{
+				ChildProvSize = 1;
+			}
 			else if (InNodeInfos[ChildId].PassedFilters && Settings->IsShowDuplicates())
 			{
 				ChildProvSize = 1;
@@ -358,7 +444,7 @@ void UEdGraph_ReferenceViewer::RecursivelyFilterNodeInfos(const FAssetIdentifier
 	bool PassedAssetTypeFilter = FilterCollection && Settings->GetFiltersEnabled() ? FilterCollection->PassesAllFilters(InNodeInfos[InAssetId]) : true;
 	bool PassedSearchTextFilter = IsAssetPassingSearchTextFilter(InAssetId);
 
-	bool PassedAllFilters = PassedAssetTypeFilter && PassedSearchTextFilter;	
+	bool PassedAllFilters = Settings->GetFindPathEnabled()  || (PassedAssetTypeFilter && PassedSearchTextFilter);
 
 	InNodeInfos[InAssetId].ChildProvisionSize = NewProvisionSize > 0 ? NewProvisionSize : (PassedAllFilters ? 1 : 0);
 	InNodeInfos[InAssetId].PassedFilters = PassedAllFilters;
@@ -494,7 +580,7 @@ UEdGraph_ReferenceViewer::RecursivelyPopulateNodeInfos(bool bInReferencers, cons
 	check(Identifiers.Num() > 0);
 	int32 ProvisionSize = 0;
 	const FAssetIdentifier& InAssetId = Identifiers[0];
-	if (InMaxDepth > 0 && InCurrentDepth < InMaxDepth)
+	if (!ExceedsMaxSearchDepth(InCurrentDepth, InMaxDepth))
 	{
 		TMap<FAssetIdentifier, EDependencyPinCategory> ReferenceLinks;
 		GetSortedLinks(Identifiers, bInReferencers, GetReferenceSearchFlags(false), ReferenceLinks);
@@ -531,9 +617,29 @@ UEdGraph_ReferenceViewer::RecursivelyPopulateNodeInfos(bool bInReferencers, cons
 	InNodeInfos[InAssetId].ChildProvisionSize = ProvisionSize > 0 ? ProvisionSize : 1;
 }
 
-void UEdGraph_ReferenceViewer::GatherAssetData(const TSet<FName>& AllPackageNames, TMap<FName, FAssetData>& OutPackageToAssetDataMap) const
+void UEdGraph_ReferenceViewer::GatherAssetData(TMap<FAssetIdentifier, FReferenceNodeInfo>& InNodeInfos)
 {
-	UE::AssetRegistry::GetAssetForPackages(AllPackageNames.Array(), OutPackageToAssetDataMap);
+	// Grab the list of packages
+	TSet<FName> PackageNames;
+	for (TPair<FAssetIdentifier, FReferenceNodeInfo>&  InfoPair : InNodeInfos)
+	{
+		FAssetIdentifier& AssetId = InfoPair.Key;
+		if (!AssetId.IsValue() && !AssetId.PackageName.IsNone())
+		{
+			PackageNames.Add(AssetId.PackageName);
+		}
+	}
+
+	// Retrieve the AssetData from the Registry
+	TMap<FName, FAssetData> PackagesToAssetDataMap;
+	UE::AssetRegistry::GetAssetForPackages(PackageNames.Array(), PackagesToAssetDataMap);
+
+
+	// Populate the AssetData back into the NodeInfos
+	for (TPair<FAssetIdentifier, FReferenceNodeInfo>&  InfoPair : InNodeInfos)
+	{
+		InfoPair.Value.AssetData = PackagesToAssetDataMap.FindRef(InfoPair.Key.PackageName);
+	}
 }
 
 UEdGraphNode_Reference* UEdGraph_ReferenceViewer::RecursivelyCreateNodes(bool bInReferencers, const FAssetIdentifier& InAssetId, const FIntPoint& InNodeLoc, const FAssetIdentifier& InParentId, UEdGraphNode_Reference* InParentNode, TMap<FAssetIdentifier, FReferenceNodeInfo>& InNodeInfos, int32 InCurrentDepth, int32 InMaxDepth, bool bIsRoot)
@@ -543,7 +649,7 @@ UEdGraphNode_Reference* UEdGraph_ReferenceViewer::RecursivelyCreateNodes(bool bI
 	const FReferenceNodeInfo& NodeInfo = InNodeInfos[InAssetId];
 	int32 NodeProvSize = 1;
 
-	UEdGraphNode_Reference* NewNode = NULL;
+	UEdGraphNode_Reference* NewNode = nullptr;
 	if (bIsRoot)
 	{
 		NewNode = InParentNode;
@@ -560,7 +666,7 @@ UEdGraphNode_Reference* UEdGraph_ReferenceViewer::RecursivelyCreateNodes(bool bI
 
 	bool bIsFirstOccurance = bIsRoot || NodeInfo.IsFirstParent(InParentId);
 	FIntPoint ChildLoc = InNodeLoc;
-	if (InMaxDepth > 0 && (InCurrentDepth < InMaxDepth) && bIsFirstOccurance) // Only expand the first parent
+	if (!ExceedsMaxSearchDepth(InCurrentDepth, InMaxDepth) && bIsFirstOccurance) // Only expand the first parent
 	{
 
 		// position the children nodes
@@ -587,6 +693,10 @@ UEdGraphNode_Reference* UEdGraph_ReferenceViewer::RecursivelyCreateNodes(bool bI
 		   	{
 		   		ChildProvSize = InNodeInfos[ChildId].ProvisionSize(InAssetId);
 		   	}
+		   	else if (Settings->GetFindPathEnabled())
+		   	{
+		   		ChildProvSize = 1;
+		   	}
 		   	else if (InNodeInfos[ChildId].PassedFilters && Settings->IsShowDuplicates())
 		   	{
 		   		ChildProvSize = 1;
@@ -665,12 +775,28 @@ const TSharedPtr<FAssetThumbnailPool>& UEdGraph_ReferenceViewer::GetAssetThumbna
 
 bool UEdGraph_ReferenceViewer::ExceedsMaxSearchDepth(int32 Depth, int32 MaxDepth) const
 {
-	// ExceedsMaxSearchDepth requires only greater (not equal than) because, even though the Depth is 1-based indexed (similarly to Breadth), the first index (index 0) corresponds to the root object 
-	return Settings->IsSearchDepthLimited() && Depth > MaxDepth;
+	// the FindPath feature is not depth limited
+ 	if (Settings->GetFindPathEnabled())
+ 	{
+ 		return false;
+ 	}
+
+ 	else if (Settings->IsSearchDepthLimited() && (Depth > MaxDepth || MaxDepth < 1))
+ 	{
+ 		return true;
+ 	}
+
+ 	return false;
 }
 
 bool UEdGraph_ReferenceViewer::ExceedsMaxSearchBreadth(int32 Breadth) const
 {
+	// the FindPath feature is not breadth limited
+	if (Settings->GetFindPathEnabled())
+	{
+		return false;
+	}
+
 	// ExceedsMaxSearchBreadth requires greater or equal than because the Breadth is 1-based indexed
 	return Settings->IsSearchBreadthLimited() && (Breadth >=  Settings->GetSearchBreadthLimit());
 }
diff --git a/Engine/Plugins/Editor/AssetManagerEditor/Source/AssetManagerEditor/Private/ReferenceViewer/ReferenceViewerSettings.cpp b/Engine/Plugins/Editor/AssetManagerEditor/Source/AssetManagerEditor/Private/ReferenceViewer/ReferenceViewerSettings.cpp
--- a/Engine/Plugins/Editor/AssetManagerEditor/Source/AssetManagerEditor/Private/ReferenceViewer/ReferenceViewerSettings.cpp
+++ b/Engine/Plugins/Editor/AssetManagerEditor/Source/AssetManagerEditor/Private/ReferenceViewer/ReferenceViewerSettings.cpp
@@ -67,81 +67,81 @@ bool UReferenceViewerSettings::IsShowDependencies() const
 	return bIsShowDependencies;
 }
 
-void UReferenceViewerSettings::SetSearchDepthLimitEnabled(bool newEnabled)
+void UReferenceViewerSettings::SetSearchDepthLimitEnabled(bool bNewEnabled)
 {
-	bLimitSearchDepth = newEnabled;
+	bLimitSearchDepth = bNewEnabled;
 	SaveConfig();
 }
 
-void UReferenceViewerSettings::SetSearchBreadthLimitEnabled(bool newEnabled)
+void UReferenceViewerSettings::SetSearchBreadthLimitEnabled(bool bNewEnabled)
 {
-	bLimitSearchBreadth = newEnabled;
+	bLimitSearchBreadth = bNewEnabled;
 	SaveConfig();
 }
 
-void UReferenceViewerSettings::SetShowSoftReferencesEnabled(bool newEnabled)
+void UReferenceViewerSettings::SetShowSoftReferencesEnabled(bool bNewEnabled)
 {
-	bIsShowSoftReferences = newEnabled;
+	bIsShowSoftReferences = bNewEnabled;
 	SaveConfig();
 }
 
-void UReferenceViewerSettings::SetShowHardReferencesEnabled(bool newEnabled)
+void UReferenceViewerSettings::SetShowHardReferencesEnabled(bool bNewEnabled)
 {
-	bIsShowHardReferences = newEnabled;
+	bIsShowHardReferences = bNewEnabled;
 	SaveConfig();
 }
 
-void UReferenceViewerSettings::SetShowFilteredPackagesOnlyEnabled(bool newEnabled)
+void UReferenceViewerSettings::SetShowFilteredPackagesOnlyEnabled(bool bNewEnabled)
 {
-	bIsShowFilteredPackagesOnly = newEnabled;
+	bIsShowFilteredPackagesOnly = bNewEnabled;
 	SaveConfig();
 }
 
-void UReferenceViewerSettings::SetCompactModeEnabled(bool newEnabled)
+void UReferenceViewerSettings::SetCompactModeEnabled(bool bNewEnabled)
 {
-	bIsCompactMode = newEnabled;
+	bIsCompactMode = bNewEnabled;
 	SaveConfig();
 }
 
-void UReferenceViewerSettings::SetShowDuplicatesEnabled(bool newEnabled)
+void UReferenceViewerSettings::SetShowDuplicatesEnabled(bool bNewEnabled)
 {
-	bIsShowDuplicates = newEnabled;
+	bIsShowDuplicates = bNewEnabled;
 	SaveConfig();
 }
 
-void UReferenceViewerSettings::SetShowEditorOnlyReferencesEnabled(bool newEnabled)
+void UReferenceViewerSettings::SetShowEditorOnlyReferencesEnabled(bool bNewEnabled)
 {
-	bIsShowEditorOnlyReferences = newEnabled;
+	bIsShowEditorOnlyReferences = bNewEnabled;
 	SaveConfig();
 }
 
-void UReferenceViewerSettings::SetShowManagementReferencesEnabled(bool newEnabled)
+void UReferenceViewerSettings::SetShowManagementReferencesEnabled(bool bNewEnabled)
 {
-	bIsShowManagementReferences = newEnabled;
+	bIsShowManagementReferences = bNewEnabled;
 	SaveConfig();
 }
 
-void UReferenceViewerSettings::SetShowSearchableNames(bool newEnabled)
+void UReferenceViewerSettings::SetShowSearchableNames(bool bNewEnabled)
 {
-	bIsShowSearchableNames = newEnabled;
+	bIsShowSearchableNames = bNewEnabled;
 	SaveConfig();
 }
 
-void UReferenceViewerSettings::SetShowCodePackages(bool newEnabled)
+void UReferenceViewerSettings::SetShowCodePackages(bool bNewEnabled)
 {
-	bIsShowCodePackages = newEnabled;
+	bIsShowCodePackages = bNewEnabled;
 	SaveConfig();
 }
 
-void UReferenceViewerSettings::SetShowReferencers(const bool newEnabled)
+void UReferenceViewerSettings::SetShowReferencers(const bool bNewEnabled)
 {
-	bIsShowReferencers = newEnabled;
+	bIsShowReferencers = bNewEnabled;
 	SaveConfig();
 }
 
-void UReferenceViewerSettings::SetShowDependencies(const bool newEnabled)
+void UReferenceViewerSettings::SetShowDependencies(const bool bNewEnabled)
 {
-	bIsShowDependencies = newEnabled;
+	bIsShowDependencies = bNewEnabled;
 	SaveConfig();
 }
 
@@ -211,9 +211,20 @@ bool UReferenceViewerSettings::GetFiltersEnabled() const
 	return bFiltersEnabled;
 }
 
-void UReferenceViewerSettings::SetFiltersEnabled(bool newEnabled)
+void UReferenceViewerSettings::SetFiltersEnabled(bool bNewEnabled)
 {
-	bFiltersEnabled = newEnabled;
+	bFiltersEnabled = bNewEnabled;
+	SaveConfig();
+}
+
+bool UReferenceViewerSettings::GetFindPathEnabled() const
+{
+	return bFindPathEnabled;
+}
+
+void UReferenceViewerSettings::SetFindPathEnabled(bool bNewEnabled)
+{
+	bFindPathEnabled = bNewEnabled;
 	SaveConfig();
 }
 
diff --git a/Engine/Plugins/Editor/AssetManagerEditor/Source/AssetManagerEditor/Private/ReferenceViewer/SReferenceViewer.cpp b/Engine/Plugins/Editor/AssetManagerEditor/Source/AssetManagerEditor/Private/ReferenceViewer/SReferenceViewer.cpp
--- a/Engine/Plugins/Editor/AssetManagerEditor/Source/AssetManagerEditor/Private/ReferenceViewer/SReferenceViewer.cpp
+++ b/Engine/Plugins/Editor/AssetManagerEditor/Source/AssetManagerEditor/Private/ReferenceViewer/SReferenceViewer.cpp
@@ -28,6 +28,8 @@
 #include "AssetRegistry/AssetRegistryModule.h"
 #include "ICollectionManager.h"
 #include "CollectionManagerModule.h"
+#include "IContentBrowserSingleton.h"
+#include "ContentBrowserModule.h"
 #include "Editor.h"
 #include "AssetManagerEditorCommands.h"
 #include "EditorWidgetsModule.h"
@@ -59,6 +61,8 @@ bool IsAssetIdentifierPassingSearchTextFilter(const FAssetIdentifier& InNode, co
 
 SReferenceViewer::~SReferenceViewer()
 {
+	Settings->SetFindPathEnabled(false); 
+
 	if (!GExitPurge)
 	{
 		if ( ensure(GraphObj) )
@@ -162,7 +166,7 @@ void SReferenceViewer::Construct(const FArguments& InArgs)
 
 				// Path
 				+SHorizontalBox::Slot()
-				.Padding(0, 0)
+				.Padding(0, 0, 4, 0)
 				.FillWidth(1.f)
 				.VAlign(VAlign_Center)
 				[
@@ -178,6 +182,20 @@ void SReferenceViewer::Construct(const FArguments& InArgs)
 						.Style(FAppStyle::Get(), "ReferenceViewer.PathText")
 					]
 				]
+				+SHorizontalBox::Slot()
+				.Padding(0, 7, 4, 8)
+				.FillWidth(1.0)
+				.VAlign(VAlign_Fill)
+				[
+					SAssignNew(FindPathAssetPicker, SComboButton)
+					.OnGetMenuContent(this, &SReferenceViewer::GenerateFindPathAssetPickerMenu)
+					.Visibility_Lambda([this]() { return !Settings->GetFindPathEnabled() ? EVisibility::Collapsed : EVisibility::Visible; })
+					.ButtonContent()
+					[
+						SNew(STextBlock)
+						.Text_Lambda([this] { return FindPathAssetId.IsValid() ? FText::FromString(FindPathAssetId.ToString()) : LOCTEXT("ChooseTargetAsset", "Choose a target asset ... "); } )
+					]
+				]
 			]
 		]
 
@@ -209,6 +227,8 @@ void SReferenceViewer::Construct(const FArguments& InArgs)
 			.Padding(8)
 			[
 				SNew(SHorizontalBox)
+				.Visibility_Lambda([this]() { return ( Settings->GetFindPathEnabled() ? EVisibility::Collapsed : EVisibility::Visible); })
+
 				+SHorizontalBox::Slot()
 				.AutoWidth()
 				[
@@ -625,6 +645,9 @@ void SReferenceViewer::OnNodeDoubleClicked(UEdGraphNode* Node)
 
 	if (!bFoundOverflow)
 	{
+		// turn off the find path tool if the user is wanting to center on another node
+		Settings->SetFindPathEnabled(false);
+
 		TSet<UObject*> Nodes;
 		Nodes.Add(Node);
 		ReCenterGraphOnNodes( Nodes );
@@ -679,11 +702,13 @@ bool SReferenceViewer::IsForwardEnabled() const
 
 void SReferenceViewer::BackClicked()
 {
+	Settings->SetFindPathEnabled(false);
 	HistoryManager.GoBack();
 }
 
 void SReferenceViewer::ForwardClicked()
 {
+	Settings->SetFindPathEnabled(false);
 	HistoryManager.GoForward();
 }
 
@@ -1135,7 +1160,7 @@ void SReferenceViewer::OnShowDuplicatesChanged()
 
 bool SReferenceViewer::IsShowDuplicatesChecked() const
 {
-	return Settings->IsShowDuplicates();
+	return Settings->GetFindPathEnabled() || Settings->IsShowDuplicates();
 }
 
 void SReferenceViewer::OnShowEditorOnlyReferencesChanged()
@@ -1324,11 +1349,10 @@ void SReferenceViewer::RegisterActions()
 		FIsActionChecked::CreateSP(this, &SReferenceViewer::IsShowCodePackagesChecked),
 		FIsActionButtonVisible::CreateLambda([this] { return bShowShowCodePackages; }));
 
-
 	ReferenceViewerActions->MapAction(
 		FAssetManagerEditorCommands::Get().ShowDuplicates,
 		FExecuteAction::CreateSP(this, &SReferenceViewer::OnShowDuplicatesChanged),
-		FCanExecuteAction(),	
+		FCanExecuteAction::CreateLambda([this] { return !Settings->GetFindPathEnabled(); }),
 		FIsActionChecked::CreateSP(this, &SReferenceViewer::IsShowDuplicatesChecked));
 
 	ReferenceViewerActions->MapAction(
@@ -1338,7 +1362,6 @@ void SReferenceViewer::RegisterActions()
 		FIsActionChecked::CreateSP(this, &SReferenceViewer::IsCompactModeChecked),
 		FIsActionButtonVisible::CreateLambda([this] { return bShowCompactMode; }));
 
-
 	ReferenceViewerActions->MapAction(
 		FAssetManagerEditorCommands::Get().FilterSearch,
 		FExecuteAction::CreateSP(this, &SReferenceViewer::OnShowFilteredPackagesOnlyChanged),
@@ -1432,8 +1455,8 @@ void SReferenceViewer::RegisterActions()
 				GraphObj->RefilterGraph();
 			}
 		}),
-		FCanExecuteAction(),	
-		FIsActionChecked::CreateLambda([this] {return Settings->GetFiltersEnabled();}));
+		FCanExecuteAction::CreateLambda([this] { return !Settings->GetFindPathEnabled(); }),
+		FIsActionChecked::CreateLambda([this] {return !Settings->GetFindPathEnabled() && Settings->GetFiltersEnabled();}));
 
 	ReferenceViewerActions->MapAction(
 		FAssetManagerEditorCommands::Get().AutoFilters,
@@ -1445,8 +1468,28 @@ void SReferenceViewer::RegisterActions()
 				GraphObj->RefilterGraph();
 			}
 		}),
-		FCanExecuteAction::CreateLambda([this] {return Settings->GetFiltersEnabled();}),
-		FIsActionChecked::CreateLambda([this] {return Settings->AutoUpdateFilters();}));
+		FCanExecuteAction::CreateLambda([this] {return !Settings->GetFindPathEnabled() && Settings->GetFiltersEnabled();}),
+		FIsActionChecked::CreateLambda([this] { return !Settings->GetFindPathEnabled() && Settings->AutoUpdateFilters();}));
+
+	ReferenceViewerActions->MapAction(
+		FAssetManagerEditorCommands::Get().FindPath,
+		FExecuteAction::CreateLambda([this] 
+		{
+			bool bWasEnabled = Settings->GetFindPathEnabled();
+			Settings->SetFindPathEnabled(!bWasEnabled);
+
+			if (!bWasEnabled && !FindPathAssetId.IsValid())
+			{
+				FindPathAssetPicker->SetIsOpen(true);
+			}
+			
+			GraphObj->RebuildGraph();
+
+			RegisterActiveTimer(0.1f, FWidgetActiveTimerDelegate::CreateSP(this, &SReferenceViewer::TriggerZoomToFit));
+
+		}),
+		FCanExecuteAction::CreateLambda([this] { return GraphObj ? GraphObj->GetCurrentGraphRootIdentifiers().Num() == 1 : false; }),
+		FIsActionChecked::CreateLambda([this] { return Settings->GetFindPathEnabled();}));
 
 	ReferenceViewerActions->MapAction(
 		FAssetManagerEditorCommands::Get().CopyPaths,
@@ -2052,7 +2095,6 @@ TSharedRef<SWidget> SReferenceViewer::MakeToolBar()
 		TAttribute<FText>::CreateSP(this, &SReferenceViewer::GetHistoryBackTooltip),
 		FSlateIcon(FReferenceViewerStyle::Get().GetStyleSetName(), "Icons.ArrowLeft"));
 
-
 	ToolBarBuilder.AddToolBarButton(
 		FUIAction(
 			FExecuteAction::CreateSP(this, &SReferenceViewer::ForwardClicked),
@@ -2063,6 +2105,12 @@ TSharedRef<SWidget> SReferenceViewer::MakeToolBar()
 		TAttribute<FText>::CreateSP(this, &SReferenceViewer::GetHistoryForwardTooltip),
 		FSlateIcon(FReferenceViewerStyle::Get().GetStyleSetName(), "Icons.ArrowRight"));
 
+	ToolBarBuilder.AddToolBarButton(FAssetManagerEditorCommands::Get().FindPath,
+		NAME_None,
+		TAttribute<FText>(),
+		TAttribute<FText>(), 
+		FSlateIcon(FReferenceViewerStyle::Get().GetStyleSetName(), "BlueprintEditor.FindInBlueprint"));
+
 	ToolBarBuilder.AddSeparator();
 
 	ToolBarBuilder.AddComboButton( 
@@ -2076,7 +2124,16 @@ TSharedRef<SWidget> SReferenceViewer::MakeToolBar()
 	ToolBarBuilder.AddToolBarButton(FAssetManagerEditorCommands::Get().ShowDuplicates,
 		NAME_None,
 		TAttribute<FText>(),
-		TAttribute<FText>(),
+		TAttribute<FText>::CreateLambda([this]() -> FText
+		{ 
+			if (Settings->GetFindPathEnabled())
+			{
+				return LOCTEXT("DuplicatesDisabledTooltip", "Duplicates are always shown when using the Find Path tool.");
+			}
+
+			return FAssetManagerEditorCommands::Get().ShowDuplicates->GetDescription();
+		}),
+
 		FSlateIcon(FReferenceViewerStyle::Get().GetStyleSetName(), "Icons.Duplicate"));
 
 	ToolBarBuilder.AddSeparator();
@@ -2084,13 +2141,31 @@ TSharedRef<SWidget> SReferenceViewer::MakeToolBar()
 	ToolBarBuilder.AddToolBarButton(FAssetManagerEditorCommands::Get().Filters,
 		NAME_None,
 		TAttribute<FText>(),
-		TAttribute<FText>(),
+		TAttribute<FText>::CreateLambda([this]() -> FText
+		{ 
+			if (Settings->GetFindPathEnabled())
+			{
+				return LOCTEXT("FiltersDisabledTooltip", "Filtering is disabled when using the Find Path tool.");
+			}
+
+			return FAssetManagerEditorCommands::Get().Filters->GetDescription();
+		}),
+
 		FSlateIcon(FReferenceViewerStyle::Get().GetStyleSetName(), "Icons.Filters"));
 
 	ToolBarBuilder.AddToolBarButton(FAssetManagerEditorCommands::Get().AutoFilters,
 		NAME_None,
 		TAttribute<FText>(),
-		TAttribute<FText>(),
+		TAttribute<FText>::CreateLambda([this]() -> FText
+		{ 
+			if (Settings->GetFindPathEnabled())
+			{
+				return LOCTEXT("AutoFiltersDisabledTooltip", "AutoFiltering is disabled when using the Find Path tool.");
+			}
+
+			return FAssetManagerEditorCommands::Get().AutoFilters->GetDescription();
+		}),
+
 		FSlateIcon(FReferenceViewerStyle::Get().GetStyleSetName(), "Icons.AutoFilters"));
 
 	ToolBarBuilder.EndSection();
@@ -2099,4 +2174,62 @@ TSharedRef<SWidget> SReferenceViewer::MakeToolBar()
 	return ToolBarBuilder.MakeWidget();
 }
 
+TSharedRef<SWidget> SReferenceViewer::GenerateFindPathAssetPickerMenu()
+{
+	FAssetPickerConfig AssetPickerConfig;
+	AssetPickerConfig.Filter.bRecursiveClasses = true;
+	AssetPickerConfig.OnAssetSelected = FOnAssetSelected::CreateSP(this, &SReferenceViewer::OnFindPathAssetSelected);
+	AssetPickerConfig.OnAssetEnterPressed = FOnAssetEnterPressed::CreateSP(this, &SReferenceViewer::OnFindPathAssetEnterPressed);
+	AssetPickerConfig.InitialAssetViewType = EAssetViewType::List;
+	AssetPickerConfig.bAllowNullSelection = true;
+	AssetPickerConfig.bFocusSearchBoxWhenOpened = true;
+	AssetPickerConfig.bAllowDragging = false;
+
+	FContentBrowserModule& ContentBrowserModule = FModuleManager::Get().LoadModuleChecked<FContentBrowserModule>(TEXT("ContentBrowser"));
+
+
+	return SNew(SBox)
+	.HeightOverride(500)
+	[
+		SNew( SBorder )
+		.BorderImage( FAppStyle::GetBrush("Menu.Background") )
+		[
+			ContentBrowserModule.Get().CreateAssetPicker(AssetPickerConfig)
+		]
+	];
+}
+
+void SReferenceViewer::OnFindPathAssetSelected( const FAssetData& AssetData )
+{
+	FindPathAssetPicker->SetIsOpen(false);
+
+	FindPathAssetId = FAssetIdentifier(AssetData.PackageName);
+
+	const TArray<FAssetIdentifier>& CurrentGraphRootIdentifiers = GraphObj->GetCurrentGraphRootIdentifiers();
+	if (!CurrentGraphRootIdentifiers.IsEmpty())
+	{
+		GraphObj->FindPath(CurrentGraphRootIdentifiers[0], FindPathAssetId);
+	}
+
+	RegisterActiveTimer(0.1f, FWidgetActiveTimerDelegate::CreateSP(this, &SReferenceViewer::TriggerZoomToFit));
+}
+
+void SReferenceViewer::OnFindPathAssetEnterPressed( const TArray<FAssetData>& AssetData )
+{
+	FindPathAssetPicker->SetIsOpen(false);
+
+	if (!AssetData.IsEmpty())
+	{
+		FindPathAssetId = FAssetIdentifier(AssetData[0].PackageName);
+
+		const TArray<FAssetIdentifier>& CurrentGraphRootIdentifiers = GraphObj->GetCurrentGraphRootIdentifiers();
+		if (!CurrentGraphRootIdentifiers.IsEmpty())
+		{
+			GraphObj->FindPath(CurrentGraphRootIdentifiers[0], FindPathAssetId);
+		}
+	}
+
+	RegisterActiveTimer(0.1f, FWidgetActiveTimerDelegate::CreateSP(this, &SReferenceViewer::TriggerZoomToFit));
+} 
+
 #undef LOCTEXT_NAMESPACE
\ No newline at end of file
diff --git a/Engine/Plugins/Editor/AssetManagerEditor/Source/AssetManagerEditor/Private/ReferenceViewer/SReferenceViewer.h b/Engine/Plugins/Editor/AssetManagerEditor/Source/AssetManagerEditor/Private/ReferenceViewer/SReferenceViewer.h
--- a/Engine/Plugins/Editor/AssetManagerEditor/Source/AssetManagerEditor/Private/ReferenceViewer/SReferenceViewer.h
+++ b/Engine/Plugins/Editor/AssetManagerEditor/Source/AssetManagerEditor/Private/ReferenceViewer/SReferenceViewer.h
@@ -175,6 +175,14 @@ private:
 	bool CanZoomToFit() const;
 	void OnFind();
 
+	/** Find Path */
+	void OnFindPathAssetSelected();
+	TSharedRef<SWidget> GenerateFindPathAssetPickerMenu();
+	void OnFindPathAssetSelected( const FAssetData& AssetData );
+	void OnFindPathAssetEnterPressed( const TArray<FAssetData>& AssetData );
+	TSharedPtr<SComboButton> FindPathAssetPicker;
+	FAssetIdentifier FindPathAssetId;
+
 	/** Handlers for searching */
 	void HandleOnSearchTextChanged(const FText& SearchText);
 	void HandleOnSearchTextCommitted(const FText& SearchText, ETextCommit::Type CommitType);
diff --git a/Engine/Plugins/Editor/AssetManagerEditor/Source/AssetManagerEditor/Public/ReferenceViewer/EdGraph_ReferenceViewer.h b/Engine/Plugins/Editor/AssetManagerEditor/Source/AssetManagerEditor/Public/ReferenceViewer/EdGraph_ReferenceViewer.h
--- a/Engine/Plugins/Editor/AssetManagerEditor/Source/AssetManagerEditor/Public/ReferenceViewer/EdGraph_ReferenceViewer.h
+++ b/Engine/Plugins/Editor/AssetManagerEditor/Source/AssetManagerEditor/Public/ReferenceViewer/EdGraph_ReferenceViewer.h
@@ -125,7 +125,7 @@ private:
 	void RecursivelyFilterNodeInfos(const FAssetIdentifier& InAssetId, TMap<FAssetIdentifier, FReferenceNodeInfo>& NodeInfos, int32 CurrentDepth, int32 MaxDepth);
 
 	/* Searches for the AssetData for the list of packages derived from the AssetReferences  */
-	void GatherAssetData(const TSet<FName>& AllPackageNames, TMap<FName, FAssetData>& OutPackageToAssetDataMap) const;
+	void GatherAssetData(TMap<FAssetIdentifier, FReferenceNodeInfo>& InNodeInfos);
 
 	/* Uses the NodeInfos map to generate and layout the graph nodes */
 	UEdGraphNode_Reference* RecursivelyCreateNodes(
@@ -152,6 +152,9 @@ private:
 	bool IsPackageIdentifierPassingFilter(const FAssetIdentifier& InAssetIdentifier) const;
 	bool IsAssetPassingSearchTextFilter(const FAssetIdentifier& InAssetIdentifier) const;
 
+	UEdGraphNode_Reference* FindPath(const FAssetIdentifier& RootId, const FAssetIdentifier& TargetId);
+	bool FindPath_Recursive(bool bInReferencers, const FAssetIdentifier& InAssetId, const FAssetIdentifier& Target, TMap<FAssetIdentifier, FReferenceNodeInfo>& InNodeInfos, TSet<FAssetIdentifier>& Visited);
+
 private:
 	/** Pool for maintaining and rendering thumbnails */
 	TSharedPtr<FAssetThumbnailPool> AssetThumbnailPool;
@@ -193,5 +196,7 @@ private:
 	/* A delegate to notify when the underlying assets changed (usually through a root or depth change) */
 	FSimpleMulticastDelegate OnAssetsChangedDelegate;
 
+	FAssetIdentifier TargetIdentifier;
+
 	friend SReferenceViewer;
 };
diff --git a/Engine/Plugins/Editor/AssetManagerEditor/Source/AssetManagerEditor/Public/ReferenceViewer/ReferenceViewerSettings.h b/Engine/Plugins/Editor/AssetManagerEditor/Source/AssetManagerEditor/Public/ReferenceViewer/ReferenceViewerSettings.h
--- a/Engine/Plugins/Editor/AssetManagerEditor/Source/AssetManagerEditor/Public/ReferenceViewer/ReferenceViewerSettings.h
+++ b/Engine/Plugins/Editor/AssetManagerEditor/Source/AssetManagerEditor/Public/ReferenceViewer/ReferenceViewerSettings.h
@@ -39,7 +39,7 @@ public:
 	GENERATED_BODY()
 
 	bool IsSearchDepthLimited() const;
-	void SetSearchDepthLimitEnabled(bool newEnabled);
+	void SetSearchDepthLimitEnabled(bool bNewEnabled);
 
 	bool IsShowReferencers() const;
 	void SetShowReferencers(const bool bShouldShowReferencers);
@@ -54,7 +54,7 @@ public:
 	void SetSearchDependencyDepthLimit(int32 NewDepthLimit, bool bSaveConfig = true);
 
 	bool IsSearchBreadthLimited() const;
-	void SetSearchBreadthLimitEnabled(bool newEnabled);
+	void SetSearchBreadthLimitEnabled(bool bNewEnabled);
 
 	int32 GetSearchBreadthLimit() const;
 	void SetSearchBreadthLimit(int32 NewBreadthLimit);
@@ -63,44 +63,46 @@ public:
 	void SetEnableCollectionFilter(bool bEnabled);
 
 	bool IsShowSoftReferences() const;
-	void SetShowSoftReferencesEnabled(bool newEnabled);
+	void SetShowSoftReferencesEnabled(bool bNewEnabled);
 
 	bool IsShowHardReferences() const;
-	void SetShowHardReferencesEnabled(bool newEnabled);
+	void SetShowHardReferencesEnabled(bool bNewEnabled);
 
 	bool IsShowEditorOnlyReferences() const;
-	void SetShowEditorOnlyReferencesEnabled(bool newEnabled);
+	void SetShowEditorOnlyReferencesEnabled(bool bNewEnabled);
 
 	bool IsShowManagementReferences() const;
-	void SetShowManagementReferencesEnabled(bool newEnabled);
+	void SetShowManagementReferencesEnabled(bool bNewEnabled);
 
 	bool IsShowSearchableNames() const;
-	void SetShowSearchableNames(bool newEnabled);
+	void SetShowSearchableNames(bool bNewEnabled);
 
 	bool IsShowCodePackages() const;
-	void SetShowCodePackages(bool newEnabled);
+	void SetShowCodePackages(bool bNewEnabled);
 
 	bool IsShowDuplicates() const;
-	void SetShowDuplicatesEnabled(bool newEnabled);
+	void SetShowDuplicatesEnabled(bool bNewEnabled);
 
 	bool IsShowFilteredPackagesOnly() const;
-	void SetShowFilteredPackagesOnlyEnabled(bool newEnabled);
+	void SetShowFilteredPackagesOnlyEnabled(bool bNewEnabled);
 
 	bool IsCompactMode() const;
-	void SetCompactModeEnabled(bool newEnabled);
+	void SetCompactModeEnabled(bool bNewEnabled);
 
 	bool IsShowPath() const;
-	void SetShowPathEnabled(bool newEnabled);
+	void SetShowPathEnabled(bool bNewEnabled);
 
 	bool GetFiltersEnabled() const;
-	void SetFiltersEnabled(bool newEnabled);
+	void SetFiltersEnabled(bool bNewEnabled);
 
 	bool AutoUpdateFilters() const;
-	void SetAutoUpdateFilters(bool newEnabled);
+	void SetAutoUpdateFilters(bool bNewEnabled);
 
 	const TArray<FilterState>& GetUserFilters() const;
 	void SetUserFilters(TArray<FilterState>& InFilters);
 
+	bool GetFindPathEnabled() const;
+	void SetFindPathEnabled(bool bNewEnabled);
 
 private:
 	/* Whether to limit the search depth for Referencers & Dependencies */
@@ -186,4 +188,9 @@ private:
 	/* The list of filters the user has built up */
 	UPROPERTY(config)
 	TArray<FilterState> UserFilters;
+
+	/* Enable the Find Path tool */
+	UPROPERTY(config)
+	bool bFindPathEnabled;
+
 };
