Editor: map "Choose File To Save" on Ctrl+Alt+Shift+S by default.

#rnx

diff --git a/Engine/Plugins/Editor/AssetManagerEditor/Source/AssetManagerEditor/Private/AssetManagerEditorCommands.cpp b/Engine/Plugins/Editor/AssetManagerEditor/Source/AssetManagerEditor/Private/AssetManagerEditorCommands.cpp
--- a/Engine/Plugins/Editor/AssetManagerEditor/Source/AssetManagerEditor/Private/AssetManagerEditorCommands.cpp
+++ b/Engine/Plugins/Editor/AssetManagerEditor/Source/AssetManagerEditor/Private/AssetManagerEditorCommands.cpp
@@ -19,7 +19,7 @@ void FAssetManagerEditorCommands::RegisterCommands()
 {
 	UI_COMMAND(ViewReferences, "Reference Viewer...", "Launches the reference viewer showing the selected assets' references", EUserInterfaceActionType::Button, FInputChord(EModifierKey::Shift | EModifierKey::Alt, EKeys::R));
 	UI_COMMAND(ViewSizeMap, "Size Map...", "Displays an interactive map showing the approximate size of this asset and everything it references", EUserInterfaceActionType::Button, FInputChord(EModifierKey::Shift | EModifierKey::Alt, EKeys::M));
-	UI_COMMAND(ViewShaderCookStatistics, "Shader Cook Statistics...", "Show Shader CookStatistics", EUserInterfaceActionType::Button, FInputChord(EModifierKey::Shift | EModifierKey::Alt | EModifierKey::Control, EKeys::S));
+	UI_COMMAND(ViewShaderCookStatistics, "Shader Cook Statistics...", "Show Shader CookStatistics", EUserInterfaceActionType::Button, FInputChord());
 	
 	UI_COMMAND(ViewAssetAudit, "Audit Assets...", "Opens the Asset Audit UI and displays information about the selected assets", EUserInterfaceActionType::Button, FInputChord(EModifierKey::Shift | EModifierKey::Alt, EKeys::A));
 
diff --git a/Engine/Source/Editor/MainFrame/Private/Frame/MainFrameActions.cpp b/Engine/Source/Editor/MainFrame/Private/Frame/MainFrameActions.cpp
--- a/Engine/Source/Editor/MainFrame/Private/Frame/MainFrameActions.cpp
+++ b/Engine/Source/Editor/MainFrame/Private/Frame/MainFrameActions.cpp
@@ -102,7 +102,7 @@ void FMainFrameCommands::RegisterCommands()
 	UI_COMMAND( SaveAll, "Save All", "Saves all unsaved levels and assets to disk", EUserInterfaceActionType::Button, FInputChord( EModifierKey::Control | EModifierKey::Shift, EKeys::S ) );
 	ActionList->MapAction( SaveAll, FExecuteAction::CreateStatic( &FMainFrameActionCallbacks::SaveAll ), FCanExecuteAction::CreateStatic( &FMainFrameActionCallbacks::CanSaveWorld ) );
 
-	UI_COMMAND( ChooseFilesToSave, "Choose Files to Save...", "Opens a dialog with save options for content and levels", EUserInterfaceActionType::Button, FInputChord() );
+	UI_COMMAND( ChooseFilesToSave, "Choose Files to Save...", "Opens a dialog with save options for content and levels", EUserInterfaceActionType::Button, FInputChord( EModifierKey::Control | EModifierKey::Alt | EModifierKey::Shift, EKeys::S ) );
 	ActionList->MapAction( ChooseFilesToSave, FExecuteAction::CreateStatic( &FMainFrameActionCallbacks::ChoosePackagesToSave ), FCanExecuteAction::CreateStatic( &FMainFrameActionCallbacks::CanSaveWorld ) );
 
 	UI_COMMAND( ViewChangelists, "View Changelists", "Opens a dialog displaying current changelists.", EUserInterfaceActionType::Button, FInputChord() );
