Move Source control menu bindings into AssetManagerEditor so they appear on all menus that operate on assets (i.e. Asset menu in editors).
* Force uncached status update if any items in selected aren't up-to-date
* Removed Refresh option since it's done every time the menu is opened
* Make Source control menu update state asynchronously so menu doesn't stall
* Update Checkout to Sync and Checkout so it can all be done at once
* Add Partial Checkout state when files are selected and show users locking files in tooltip
* Add options to Make Writeable and Revert Writeable
* Add error state if source control is unavailable
* Hide options, like checkout, that are not available for the current SCC provider
* Always show all options so users can learn bettern muscle memory
* Add Warning when doing any source control option on selections larger than 25 (currently hardcoded)
* There's still some slight difference in behavior between file and folder selection. If only files are selected the async update will not show any options availabe until the full scan is done and there are some additional features like partial checkout. If there are any folders then options become available the moment one asset is available for an action.

#tests Tested all button scenarios, tested small and large sets of folders, assets and a mix.

diff --git a/Engine/Plugins/Animation/ControlRig/Source/ControlRigEditor/Private/EditMode/SControlRigBaseListWidget.h b/Engine/Plugins/Animation/ControlRig/Source/ControlRigEditor/Private/EditMode/SControlRigBaseListWidget.h
--- a/Engine/Plugins/Animation/ControlRig/Source/ControlRigEditor/Private/EditMode/SControlRigBaseListWidget.h
+++ b/Engine/Plugins/Animation/ControlRig/Source/ControlRigEditor/Private/EditMode/SControlRigBaseListWidget.h
@@ -23,7 +23,6 @@
 class UControlRig;
 class FControlRigEditMode;
 class UControlRigPoseAsset;
-class FAssetFileContextMenu;
 
 namespace ESelectedControlAsset
 {
@@ -121,8 +120,6 @@ private:
 	TSharedPtr<SWidget> PathPicker;
 	TSharedPtr<SWidget> AssetPicker;
 
-	TSharedPtr<FAssetFileContextMenu> AssetFileContextMenu;
-
 	static FString CurrentlySelectedInternalPath;
 	FString CurrentlyEnteredAssetName;
 	bool bLastInputValidityCheckSuccessful;
diff --git a/Engine/Plugins/Editor/AssetManagerEditor/Source/AssetManagerEditor/AssetManagerEditor.Build.cs b/Engine/Plugins/Editor/AssetManagerEditor/Source/AssetManagerEditor/AssetManagerEditor.Build.cs
--- a/Engine/Plugins/Editor/AssetManagerEditor/Source/AssetManagerEditor/AssetManagerEditor.Build.cs
+++ b/Engine/Plugins/Editor/AssetManagerEditor/Source/AssetManagerEditor/AssetManagerEditor.Build.cs
@@ -46,6 +46,8 @@ public class AssetManagerEditor : ModuleRules
 				"TreeMap",
 				"ToolMenus",
 				"ToolWidgets",
+				"SourceControl",
+				"SourceControlWindows",
 			}
 		);
 	}
diff --git a/Engine/Plugins/Editor/AssetManagerEditor/Source/AssetManagerEditor/Private/AssetManagerEditorModule.cpp b/Engine/Plugins/Editor/AssetManagerEditor/Source/AssetManagerEditor/Private/AssetManagerEditorModule.cpp
--- a/Engine/Plugins/Editor/AssetManagerEditor/Source/AssetManagerEditor/Private/AssetManagerEditorModule.cpp
+++ b/Engine/Plugins/Editor/AssetManagerEditor/Source/AssetManagerEditor/Private/AssetManagerEditorModule.cpp
@@ -50,6 +50,7 @@
 #include "EdGraphSchema_K2.h"
 #include "SGraphPin.h"
 #include "AssetManagerEditorCommands.h"
+#include "AssetSourceControlContextMenu.h"
 #include "ReferenceViewer/SReferenceViewer.h"
 #include "ReferenceViewer/SReferenceNode.h"
 #include "ReferenceViewer/EdGraphNode_Reference.h"
@@ -380,7 +381,9 @@ private:
 	TSharedPtr<FAssetManagerGraphPanelNodeFactory> AssetManagerGraphPanelNodeFactory;
 	TSharedPtr<FAssetManagerGraphPanelPinFactory> AssetManagerGraphPanelPinFactory;
 
-	static void CreateAssetContextMenu(FToolMenuSection& InSection);
+	FAssetSourceControlContextMenu AssetSourceControlContextMenu;
+
+	void CreateAssetContextMenu(FToolMenuSection& InSection);
 	void OnExtendContentBrowserCommands(TSharedRef<FUICommandList> CommandList, FOnContentBrowserGetSelection GetSelectionDelegate);
 	void OnExtendLevelEditorCommands(TSharedRef<FUICommandList> CommandList);
 	void RegisterMenus();
@@ -847,6 +850,8 @@ void FAssetManagerEditorModule::CreateAssetContextMenu(FToolMenuSection& InSecti
 		TAttribute<FText>(), // Use command tooltip
 		FSlateIcon(FAppStyle::GetAppStyleSetName(), "MainFrame.CookContent")
 	);
+
+	AssetSourceControlContextMenu.MakeContextMenu(InSection);
 }
 
 void FAssetManagerEditorModule::OnExtendContentBrowserCommands(TSharedRef<FUICommandList> CommandList, FOnContentBrowserGetSelection GetSelectionDelegate)
@@ -932,12 +937,12 @@ void FAssetManagerEditorModule::ExtendContentBrowserAssetSelectionMenu()
 {
 	UToolMenu* Menu = UToolMenus::Get()->ExtendMenu("ContentBrowser.AssetContextMenu");
 	FToolMenuSection& Section = Menu->FindOrAddSection("AssetContextReferences");
-	FToolMenuEntry& Entry = Section.AddDynamicEntry("AssetManagerEditorViewCommands", FNewToolMenuSectionDelegate::CreateLambda([](FToolMenuSection& InSection)
+	FToolMenuEntry& Entry = Section.AddDynamicEntry("AssetManagerEditorViewCommands", FNewToolMenuSectionDelegate::CreateLambda([this](FToolMenuSection& InSection)
 	{
 		UContentBrowserAssetContextMenuContext* Context = InSection.FindContext<UContentBrowserAssetContextMenuContext>();
 		if (Context && Context->bCanBeModified && Context->SelectedAssets.Num() > 0)
 		{
-			FAssetManagerEditorModule::CreateAssetContextMenu(InSection);
+			CreateAssetContextMenu(InSection);
 		}
 	}));
 }
@@ -946,12 +951,12 @@ void FAssetManagerEditorModule::ExtendContentBrowserPathSelectionMenu()
 {
 	UToolMenu* Menu = UToolMenus::Get()->ExtendMenu("ContentBrowser.FolderContextMenu");
 	FToolMenuSection& Section = Menu->FindOrAddSection("PathContextBulkOperations");
-	FToolMenuEntry& Entry = Section.AddDynamicEntry("AssetManagerEditorViewCommands", FNewToolMenuSectionDelegate::CreateLambda([](FToolMenuSection& InSection)
+	FToolMenuEntry& Entry = Section.AddDynamicEntry("AssetManagerEditorViewCommands", FNewToolMenuSectionDelegate::CreateLambda([this](FToolMenuSection& InSection)
 	{
 		UContentBrowserFolderContext* Context = InSection.FindContext<UContentBrowserFolderContext>();
 		if (Context && Context->bCanBeModified && Context->NumAssetPaths > 0)
 		{
-			FAssetManagerEditorModule::CreateAssetContextMenu(InSection);
+			CreateAssetContextMenu(InSection);
 		}
 	}));
 }
@@ -996,7 +1001,7 @@ void FAssetManagerEditorModule::ExtendAssetEditorMenu()
 {
 	UToolMenu* Menu = UToolMenus::Get()->ExtendMenu("MainFrame.MainMenu.Asset");
 	FToolMenuSection& Section = Menu->FindOrAddSection("AssetEditorActions");
-	FToolMenuEntry& Entry = Section.AddDynamicEntry("AssetManagerEditorViewCommands", FNewToolMenuSectionDelegate::CreateLambda([](FToolMenuSection& InSection)
+	FToolMenuEntry& Entry = Section.AddDynamicEntry("AssetManagerEditorViewCommands", FNewToolMenuSectionDelegate::CreateLambda([this](FToolMenuSection& InSection)
 	{
 		UAssetEditorToolkitMenuContext* MenuContext = InSection.FindContext<UAssetEditorToolkitMenuContext>();
 		if (MenuContext && MenuContext->Toolkit.IsValid() && MenuContext->Toolkit.Pin()->IsActuallyAnAsset())
@@ -1005,7 +1010,7 @@ void FAssetManagerEditorModule::ExtendAssetEditorMenu()
 			{
 				if (IsValid(EditedAsset) && EditedAsset->IsAsset())
 				{
-					FAssetManagerEditorModule::CreateAssetContextMenu(InSection);
+					CreateAssetContextMenu(InSection);
 					break;
 				}
 			}
@@ -1019,7 +1024,7 @@ void FAssetManagerEditorModule::ExtendLevelEditorActorContextMenu()
 	UToolMenu* Menu = UToolMenus::Get()->ExtendMenu("LevelEditor.ActorContextMenu.AssetToolsSubMenu");
 
 	FToolMenuSection& Section = Menu->AddSection("AssetManagerEditorViewCommands", TAttribute<FText>(), FToolMenuInsert(NAME_None, EToolMenuInsertType::First));
-	FAssetManagerEditorModule::CreateAssetContextMenu(Section);
+	CreateAssetContextMenu(Section);
 }
 
 void FAssetManagerEditorModule::OnReloadComplete(EReloadCompleteReason Reason)
diff --git a/Engine/Plugins/Editor/AssetManagerEditor/Source/AssetManagerEditor/Private/AssetSourceControlContextMenu.cpp b/Engine/Plugins/Editor/AssetManagerEditor/Source/AssetManagerEditor/Private/AssetSourceControlContextMenu.cpp
new file mode 100644
--- /dev/null
+++ b/Engine/Plugins/Editor/AssetManagerEditor/Source/AssetManagerEditor/Private/AssetSourceControlContextMenu.cpp
@@ -0,0 +1,1135 @@
+// Copyright Epic Games, Inc. All Rights Reserved.
+
+#include "AssetSourceControlContextMenu.h"
+
+#include "Algo\Transform.h"
+#include "AssetRegistry/AssetRegistryModule.h"
+#include "AssetToolsModule.h"
+#include "AssetViewUtils.h"
+#include "Async/Async.h"
+#include "ContentBrowserMenuContexts.h"
+#include "FileHelpers.h"
+#include "ISourceControlModule.h"
+#include "Misc/FileHelper.h"
+#include "Misc/MessageDialog.h"
+#include "Misc/ScopedSlowTask.h"
+#include "SourceControlHelpers.h"
+#include "SourceControlOperations.h"
+#include "SourceControlWindows.h"
+#include "Subsystems/AssetEditorSubsystem.h"
+#include "Toolkits/AssetEditorToolkit.h"
+#include "Toolkits/AssetEditorToolkitMenuContext.h"
+#include "ToolMenu.h"
+#include "ToolMenuSection.h"
+
+#include "Styling/CoreStyle.h"
+#include "Styling/ISlateStyle.h"
+#include "Styling/ToolBarStyle.h"
+#include "ToolMenuDelegates.h"
+#include "ToolMenuEntry.h"
+#include "ToolMenuSection.h"
+#include "Widgets/Images/SThrobber.h"
+#include "Widgets/SBoxPanel.h"
+#include "Widgets/Text/STextBlock.h"
+
+#define LOCTEXT_NAMESPACE "AssetSourceControlContextMenu"
+
+namespace UE::AssetSourceControlContextMenu::Private
+{
+	DECLARE_DELEGATE_RetVal(bool, FIsAsyncProcessingActive);
+
+	FNewToolMenuCustomWidget MakeCustomWidgetDelegate(const TAttribute<FText>& Label, const TAttribute<FSlateIcon>& Icon, const FIsAsyncProcessingActive& IsAsyncProcessingActive)
+	{
+		return FNewToolMenuCustomWidget::CreateLambda([Label, Icon, IsAsyncProcessingActive](const FToolMenuContext& InContext, const FToolMenuCustomWidgetContext& WidgetContext)
+			{
+				return SNew(SHorizontalBox)
+					+ SHorizontalBox::Slot()
+					.VAlign(VAlign_Center)
+					.Padding(FMargin(2, 0, 0, 0))
+					.AutoWidth()
+					[
+						SNew(SImage)
+						.ColorAndOpacity( FSlateColor::UseSubduedForeground() )
+						.Image(TAttribute<const FSlateBrush*>::CreateLambda([Icon]()
+							{
+								return Icon.IsSet() ? Icon.Get().GetIcon() : nullptr;
+							}))
+					]
+				+ SHorizontalBox::Slot()
+					.VAlign(VAlign_Center)
+					.Padding(FMargin(2, 0, 0, 0))
+					.AutoWidth()
+					[
+						SNew(STextBlock)
+						.Text(Label)
+						.TextStyle(WidgetContext.StyleSet, ISlateStyle::Join(WidgetContext.StyleName, ".Label"))
+					]
+				+ SHorizontalBox::Slot()
+					.Padding(FMargin(8, 0, 0, 0))
+					.VAlign(VAlign_Center)
+					.AutoWidth()
+					[
+						SNew(SThrobber)
+						.PieceImage(WidgetContext.StyleSet->GetBrush("Throbber.CircleChunk"))
+						.Visibility_Lambda([IsAsyncProcessingActive]()
+						{
+							if (IsAsyncProcessingActive.Execute())
+							{
+								return EVisibility::Visible;
+							}
+							return EVisibility::Hidden;
+						})
+					];
+			});
+	};
+
+
+	FToolMenuEntry& AddAsyncMenuEntry(
+		FToolMenuSection& Section,
+		FName Name,
+		const TAttribute<FText>& Label,
+		const TAttribute<FText>& ToolTip,
+		const TAttribute<FSlateIcon>& Icon,
+		const FToolUIActionChoice& Action,
+		const FIsAsyncProcessingActive& IsAsyncProcessingActive
+	)
+	{
+		FToolMenuEntry MenuEntry = FToolMenuEntry::InitMenuEntry(
+			Name,
+			Label,
+			ToolTip,
+			Icon,
+			Action
+		);
+
+		if (IsAsyncProcessingActive.IsBound())
+		{
+			MenuEntry.MakeCustomWidget = MakeCustomWidgetDelegate(Label, Icon, IsAsyncProcessingActive);
+		}
+		return Section.AddEntry(MenuEntry);
+	}
+}
+
+class FAssetSourceControlContextMenuState : public TSharedFromThis<FAssetSourceControlContextMenuState>
+{
+public:
+	void Initialize(FToolMenuSection& InSection);
+
+	bool IsValid() const;
+
+	void Close();
+
+private:
+	bool AddSourceControlMenuOptions(FToolMenuSection& InSection);
+	void FillSourceControlSubMenu(UToolMenu* Menu);
+
+	FExecuteAction ExecutionCheck(FExecuteAction&& InAction) const;
+
+	void ExecuteDiffSelected() const;
+	void ExecuteSCCMerge() const;
+	void ExecuteSCCCheckOut() const;
+	void ExecuteSCCSyncAndCheckOut() const;
+	void ExecuteSCCMakeWritable() const;
+	void ExecuteSCCOpenForAdd() const;
+	void ExecuteSCCCheckIn() const;
+	void ExecuteSCCHistory() const;
+	void ExecuteSCCDiffAgainstDepot() const;
+	void ExecuteSCCRevert() const;
+	void ExecuteSCCRevertWritable() const;
+	void ExecuteSCCSync() const;
+	void ExecuteSCCRefresh() const;
+
+	bool CanExecuteSCCMerge() const;
+	bool CanExecuteSCCCheckOut() const;
+	bool CanExecuteSCCSyncAndCheckOut() const;
+	bool CanExecuteSCCMakeWritable() const;
+	bool CanExecuteSCCOpenForAdd() const;
+	bool CanExecuteSCCCheckIn() const;
+	bool CanExecuteSCCHistory() const;
+	bool CanExecuteSCCRevert() const;
+	bool CanExecuteSCCRevertWritable() const;
+	bool CanExecuteSCCSync() const;
+	bool CanExecuteSCCDiffAgainstDepot() const;
+	bool CanExecuteDiffSelected() const;
+
+	/** Helper function to gather the package names of all selected assets */
+	void GetSelectedPackageNames(TArray<FString>& OutPackageNames) const;
+
+	/** Helper function to gather the packages containing all selected assets */
+	void GetSelectedPackages(TArray<UPackage*>& OutPackages) const;
+
+private:
+
+	enum class EAsyncState : uint8
+	{
+		None,
+		TryCacheCanExecute,
+		SCCUpdate,
+		Idle,
+		Error,
+		Cancelled,
+	};
+
+	bool IsActionEnabled(bool bEnabled) const;
+
+	bool IsStillScanning(bool bEnabled) const;
+
+	bool TrySetAsyncState(EAsyncState FromState, EAsyncState ToState);
+
+	/** Initializes some variable used to in "CanExecute" checks that won't change at runtime or are too expensive to check every frame. */
+	void NextAsyncState();
+
+	/** Callback from a perforce update call. */
+	void CacheCanExecuteVarsSCCHandler(const FSourceControlOperationRef& InOperation, ECommandResult::Type Result);
+
+	/** Try using perforce cache to update state */
+	void TryCacheCanExecuteVars(const TArray<FString>& InSelectedPaths, TArray<FString>* OptionalPathsWithUnknownState);
+
+	/** Cancel any currently running perforce operation */
+	void CancelCacheCanExecuteVars();
+
+private:
+
+	TArray<FAssetData> SelectedAssets;
+	TArray<FString> PathsWithUnknownState;
+	TArray<FString> CheckedOutUsers;
+	FText CheckedOutUsersText;
+
+	TSharedPtr<class ISourceControlOperation, ESPMode::ThreadSafe> SCCOperation;
+	std::atomic<EAsyncState> AsyncState = EAsyncState::None;
+
+	// If folders are present in the selection then we adjust some behaviors to account for the potentially large amount of files to scan
+	bool bContainsFolders = false;
+
+	/** Cached CanExecute vars */
+	bool bCanExecuteSCCMerge = false;
+	bool bCanExecuteSCCCheckOut = false;
+	bool bCanExecuteSCCSyncAndCheckOut = false;
+	bool bCanExecuteSCCMakeWritable = false;
+	bool bCanExecuteSCCOpenForAdd = false;
+	bool bCanExecuteSCCCheckIn = false;
+	bool bCanExecuteSCCHistory = false;
+	bool bCanExecuteSCCRevert = false;
+	bool bCanExecuteSCCSync = false;
+	bool bCanExecuteSCCRevertWritable = false;
+};
+
+FAssetSourceControlContextMenu::~FAssetSourceControlContextMenu()
+{
+	if (InnerState)
+	{
+		InnerState->Close();
+	}
+}
+
+void FAssetSourceControlContextMenu::MakeContextMenu(FToolMenuSection& InSection)
+{
+	if (InnerState)
+	{
+		InnerState->Close();
+	}
+
+	InnerState = MakeShareable(new FAssetSourceControlContextMenuState());
+	InnerState->Initialize(InSection);
+}
+
+void FAssetSourceControlContextMenuState::Initialize(FToolMenuSection& InSection)
+{
+	if (const UContentBrowserAssetContextMenuContext* Context = InSection.FindContext<UContentBrowserAssetContextMenuContext>())
+	{
+		if (Context->bCanBeModified)
+		{
+			SelectedAssets += Context->SelectedAssets;
+		}
+	}
+
+	if (const UAssetEditorToolkitMenuContext* MenuContext = InSection.FindContext<UAssetEditorToolkitMenuContext>())
+	{
+		if (MenuContext->Toolkit.IsValid() && MenuContext->Toolkit.Pin()->IsActuallyAnAsset())
+		{
+			for (const UObject* EditedAsset : *MenuContext->Toolkit.Pin()->GetObjectsCurrentlyBeingEdited())
+			{
+				SelectedAssets.Add(FAssetData(EditedAsset));
+			}
+		}
+	}
+
+	if (const UContentBrowserDataMenuContext_FolderMenu* ContextObject = InSection.FindContext<UContentBrowserDataMenuContext_FolderMenu>())
+	{
+		TArray<FName> SelectedPaths;
+
+		for (const FContentBrowserItem& SelectedItem : ContextObject->SelectedItems)
+		{
+			for (const FContentBrowserItemData& SelectedItemData : SelectedItem.GetInternalItems())
+			{
+				SelectedPaths.Add(SelectedItemData.GetInternalPath());
+			}
+		}
+
+		// Load the asset registry module
+		FAssetRegistryModule& AssetRegistryModule = FModuleManager::Get().LoadModuleChecked<FAssetRegistryModule>(TEXT("AssetRegistry"));
+
+		// Form a filter from the paths
+		FARFilter Filter;
+		Filter.bRecursivePaths = true;
+		Filter.PackagePaths = MoveTemp(SelectedPaths);
+
+		// Query for a list of assets in the selected paths
+		AssetRegistryModule.Get().GetAssets(Filter, SelectedAssets);
+
+		bContainsFolders = true;
+	}
+
+	if (IsValid())
+	{
+		AddSourceControlMenuOptions(InSection);
+	}
+}
+
+bool FAssetSourceControlContextMenuState::IsValid() const
+{
+	return SelectedAssets.Num() > 0;
+}
+
+void FAssetSourceControlContextMenuState::Close()
+{
+	AsyncState.store(EAsyncState::Cancelled);
+
+	if (SCCOperation.IsValid())
+	{
+		ISourceControlModule::Get().GetProvider().CancelOperation(SCCOperation.ToSharedRef());
+	}
+}
+
+bool FAssetSourceControlContextMenuState::AddSourceControlMenuOptions(FToolMenuSection& InSection)
+{
+	if (ISourceControlModule::Get().IsEnabled())
+	{
+		if (ISourceControlModule::Get().GetProvider().IsAvailable())
+		{
+			// Start work to cache execution state, this can take a while if SCC cache is cold
+			NextAsyncState();
+
+			// SCC sub menu
+			InSection.AddSubMenu(
+				"SourceControlSubMenu",
+				LOCTEXT("SourceControlSubMenuLabel", "Source Control"),
+				LOCTEXT("SourceControlSubMenuToolTip", "Source control actions."),
+				FNewToolMenuDelegate::CreateSP(this, &FAssetSourceControlContextMenuState::FillSourceControlSubMenu),
+				false,
+				FSlateIcon(FAppStyle::GetAppStyleSetName(), "SourceControl.StatusIcon.On")
+			);
+		}
+		else
+		{
+			InSection.AddSubMenu(
+				"SourceControlSubMenu",
+				LOCTEXT("SourceControlSubMenuUnavailableLabel", "Source Control (Unavailable)"),
+				LOCTEXT("SourceControlSubMenuUnavailableToolTip", "Source control servers are currently unavailable, check your network connection."),
+				FNewToolMenuDelegate(),
+				FUIAction(
+					FExecuteAction(),
+					FCanExecuteAction::CreateLambda([]() { return false; })
+				),
+				EUserInterfaceActionType::Button,
+				false,
+				FSlateIcon(FAppStyle::GetAppStyleSetName(), "SourceControl.StatusIcon.Error")
+			);
+		}
+	}
+	else
+	{
+		InSection.AddMenuEntry(
+			"SCCConnectToSourceControl",
+			LOCTEXT("SCCConnectToSourceControl", "Connect To Source Control..."),
+			LOCTEXT("SCCConnectToSourceControlTooltip", "Connect to source control to allow source control operations to be performed on content and levels."),
+			FSlateIcon(FAppStyle::GetAppStyleSetName(), "MainFrame.ConnectToSourceControl"),
+			FUIAction(
+				FExecuteAction::CreateLambda([]() { ISourceControlModule::Get().ShowLoginDialog(FSourceControlLoginClosed(), ELoginWindowMode::Modeless); }),
+				FCanExecuteAction()
+			)
+		);
+	}
+
+	// Diff selected
+	if (CanExecuteDiffSelected())
+	{
+		InSection.AddMenuEntry(
+			"DiffSelected",
+			LOCTEXT("DiffSelected", "Diff Selected"),
+			LOCTEXT("DiffSelectedTooltip", "Diff the two assets that you have selected."),
+			FSlateIcon(FAppStyle::GetAppStyleSetName(), "SourceControl.Actions.Diff"),
+			FUIAction(
+				FExecuteAction::CreateSP(this, &FAssetSourceControlContextMenuState::ExecuteDiffSelected)
+			)
+		);
+	}
+
+	return true;
+}
+
+void FAssetSourceControlContextMenuState::FillSourceControlSubMenu(UToolMenu* Menu)
+{
+	ISourceControlProvider& SourceControlProvider = ISourceControlModule::Get().GetProvider();
+
+	FToolMenuSection& Section = Menu->AddSection("AssetSourceControlActions", LOCTEXT("AssetSourceControlActionsMenuHeading", "Source Control"));
+
+	using namespace UE::AssetSourceControlContextMenu::Private;
+
+	const bool bUsesCheckout = SourceControlProvider.UsesCheckout();
+	const bool bUsesFileRevisions = SourceControlProvider.UsesFileRevisions();
+	const bool bUsesReadOnly = SourceControlProvider.UsesLocalReadOnlyState();
+
+	if (bUsesFileRevisions)
+	{
+		AddAsyncMenuEntry(Section,
+			"SCCSync",
+			LOCTEXT("SCCSync", "Sync"),
+			LOCTEXT("SCCSyncTooltip", "Updates the selected assets to the latest version in source control."),
+			FSlateIcon(FAppStyle::GetAppStyleSetName(), "SourceControl.Actions.Sync"),
+			FUIAction(
+				ExecutionCheck(FExecuteAction::CreateSP(this, &FAssetSourceControlContextMenuState::ExecuteSCCSync)),
+				FCanExecuteAction::CreateLambda([this]() { return IsActionEnabled(CanExecuteSCCSync()); })
+			),
+			FIsAsyncProcessingActive::CreateLambda([this]() { return IsStillScanning(CanExecuteSCCSync()); })
+		);
+	}
+
+	if (bUsesCheckout)
+	{
+		AddAsyncMenuEntry(Section,
+			"SCCCheckOut",
+			LOCTEXT("SCCCheckOut", "Check Out"),
+			TAttribute<FText>::CreateLambda([this]()
+				{
+					if (IsStillScanning(CanExecuteSCCCheckOut()) || CheckedOutUsers.Num() == 0)
+					{
+						return LOCTEXT("SCCCheckOutTooltip", "Check out the selected assets from source control.");
+					}
+
+					return FText::Format(LOCTEXT("SCCPartialSyncAndCheckOut", "Checks out the selected assets from source control that are not currently locked.\n\nLocked Assets:\n{0}"), CheckedOutUsersText);
+				}),
+
+			TAttribute<FSlateIcon>::CreateLambda([this]()
+				{
+					return FSlateIcon(FAppStyle::GetAppStyleSetName(), CheckedOutUsers.Num() == 0 ? "SourceControl.Actions.CheckOut" : "SourceControl.LockOverlay");
+				}),
+					FUIAction(
+						ExecutionCheck(FExecuteAction::CreateSP(this, &FAssetSourceControlContextMenuState::ExecuteSCCCheckOut)),
+						FCanExecuteAction::CreateLambda([this]() { return IsActionEnabled(CanExecuteSCCCheckOut()); })
+					),
+					FIsAsyncProcessingActive::CreateLambda([this]() { return IsStillScanning(CanExecuteSCCCheckOut()); })
+					);
+
+		AddAsyncMenuEntry(Section,
+			"SCCSyncAndCheckOut",
+			LOCTEXT("SCCSyncAndCheckOut", "Sync And Check Out"),
+			TAttribute<FText>::CreateLambda([this]()
+				{
+					if (IsStillScanning(CanExecuteSCCCheckOut()) || CheckedOutUsers.Num() == 0)
+					{
+						return LOCTEXT("SCCSyncAndCheckOutTooltip", "Sync to latest and Check out the selected assets from source control.");
+					}
+
+					return FText::Format(LOCTEXT("SCCPartialSyncAndCheckOut", "Sync to latest and Checks out the selected assets from source control that are not currently locked.\n\nLocked Assets:\n{0}"), CheckedOutUsersText);
+				}),
+
+			TAttribute<FSlateIcon>::CreateLambda([this]()
+				{
+					return FSlateIcon(FAppStyle::GetAppStyleSetName(), CheckedOutUsers.Num() == 0 ? "SourceControl.Actions.CheckOut" : "SourceControl.LockOverlay");
+				}),
+					FUIAction(
+						ExecutionCheck(FExecuteAction::CreateSP(this, &FAssetSourceControlContextMenuState::ExecuteSCCSyncAndCheckOut)),
+						FCanExecuteAction::CreateLambda([this]() { return IsActionEnabled(CanExecuteSCCSyncAndCheckOut()); })
+					),
+					FIsAsyncProcessingActive::CreateLambda([this]() { return IsStillScanning(CanExecuteSCCSyncAndCheckOut()); })
+					);
+	}
+
+	if (bUsesReadOnly)
+	{
+		AddAsyncMenuEntry(Section,
+			"SCCMakeWritable",
+			LOCTEXT("SCCMakeWritable", "Make Writable"),
+			LOCTEXT("SCCMakeWritableTooltip", "Remove read-only flag from all selected assets."),
+			FSlateIcon(FAppStyle::GetAppStyleSetName(), "SourceControl.Actions.CheckOut"),
+			FUIAction(
+				ExecutionCheck(FExecuteAction::CreateSP(this, &FAssetSourceControlContextMenuState::ExecuteSCCMakeWritable)),
+				FCanExecuteAction::CreateLambda([this]() { return IsActionEnabled(CanExecuteSCCMakeWritable()); })
+			),
+			FIsAsyncProcessingActive::CreateLambda([this]() { return IsStillScanning(CanExecuteSCCMakeWritable()); })
+		);
+	}
+
+	AddAsyncMenuEntry(Section,
+		"SCCOpenForAdd",
+		LOCTEXT("SCCOpenForAdd", "Mark For Add"),
+		LOCTEXT("SCCOpenForAddTooltip", "Adds the selected assets to source control."),
+		FSlateIcon(FAppStyle::GetAppStyleSetName(), "SourceControl.Actions.Add"),
+		FUIAction(
+			ExecutionCheck(FExecuteAction::CreateSP(this, &FAssetSourceControlContextMenuState::ExecuteSCCOpenForAdd)),
+			FCanExecuteAction::CreateLambda([this]() { return IsActionEnabled(CanExecuteSCCOpenForAdd()); })
+		),
+		FIsAsyncProcessingActive::CreateLambda([this]() { return IsStillScanning(CanExecuteSCCOpenForAdd()); })
+	);
+
+	if (bUsesFileRevisions)
+	{
+		AddAsyncMenuEntry(Section,
+			"SCCCheckIn",
+			LOCTEXT("SCCCheckIn", "Check In"),
+			LOCTEXT("SCCCheckInTooltip", "Checks the selected assets into source control."),
+			FSlateIcon(FAppStyle::GetAppStyleSetName(), "SourceControl.Actions.Submit"),
+			FUIAction(
+				ExecutionCheck(FExecuteAction::CreateSP(this, &FAssetSourceControlContextMenuState::ExecuteSCCCheckIn)),
+				FCanExecuteAction::CreateLambda([this]() { return IsActionEnabled(CanExecuteSCCCheckIn()); })
+			),
+			FIsAsyncProcessingActive::CreateLambda([this]() { return IsStillScanning(CanExecuteSCCCheckIn()); })
+		);
+	}
+
+	AddAsyncMenuEntry(Section,
+		"SCCHistory",
+		LOCTEXT("SCCHistory", "History"),
+		LOCTEXT("SCCHistoryTooltip", "Displays the source control revision history of the selected assets."),
+		FSlateIcon(FAppStyle::GetAppStyleSetName(), "SourceControl.Actions.History"),
+		FUIAction(
+			ExecutionCheck(FExecuteAction::CreateSP(this, &FAssetSourceControlContextMenuState::ExecuteSCCHistory)),
+			FCanExecuteAction::CreateLambda([this]() { return IsActionEnabled(CanExecuteSCCHistory()); })
+		),
+		FIsAsyncProcessingActive::CreateLambda([this]() { return IsStillScanning(CanExecuteSCCHistory()); })
+	);
+
+	AddAsyncMenuEntry(Section,
+		"SCCDiffAgainstDepot",
+		LOCTEXT("SCCDiffAgainstDepot", "Diff Against Depot"),
+		LOCTEXT("SCCDiffAgainstDepotTooltip", "Look at differences between the local and remote version of the selected assets."),
+		FSlateIcon(FAppStyle::GetAppStyleSetName(), "SourceControl.Actions.Diff"),
+		FUIAction(
+			ExecutionCheck(FExecuteAction::CreateSP(this, &FAssetSourceControlContextMenuState::ExecuteSCCDiffAgainstDepot)),
+			FCanExecuteAction::CreateLambda([this]() { return IsActionEnabled(CanExecuteSCCDiffAgainstDepot()); })
+		),
+		FIsAsyncProcessingActive::CreateLambda([this]() { return IsStillScanning(CanExecuteSCCDiffAgainstDepot()); })
+	);
+
+	AddAsyncMenuEntry(Section,
+		"SCCRevert",
+		LOCTEXT("SCCRevert", "Revert"),
+		LOCTEXT("SCCRevertTooltip", "Reverts the selected assets to their original state from source control."),
+		FSlateIcon(FAppStyle::GetAppStyleSetName(), "SourceControl.Actions.Revert"),
+		FUIAction(
+			ExecutionCheck(FExecuteAction::CreateSP(this, &FAssetSourceControlContextMenuState::ExecuteSCCRevert)),
+			FCanExecuteAction::CreateLambda([this]() { return IsActionEnabled(CanExecuteSCCRevert()); })
+		),
+		FIsAsyncProcessingActive::CreateLambda([this]() { return IsStillScanning(CanExecuteSCCRevert()); })
+	);
+
+	if (bUsesReadOnly)
+	{
+		AddAsyncMenuEntry(Section,
+			"SCCRevertWritable",
+			LOCTEXT("SCCRevertWritable", "Revert Writable Files"),
+			LOCTEXT("SCCRevertTooltip", "Reverts the assets that are Writable to their current state from source control. They will remain at their current revision."),
+			FSlateIcon(FAppStyle::GetAppStyleSetName(), "SourceControl.Actions.Revert"),
+			FUIAction(
+				ExecutionCheck(FExecuteAction::CreateSP(this, &FAssetSourceControlContextMenuState::ExecuteSCCRevertWritable)),
+				FCanExecuteAction::CreateLambda([this]() { return IsActionEnabled(CanExecuteSCCRevertWritable()); })
+			),
+			FIsAsyncProcessingActive::CreateLambda([this]() { return IsStillScanning(CanExecuteSCCRevertWritable()); })
+		);
+	}
+
+	// At the end since this is the only item that will appear/disappear based on the content being looked at
+	if (!bContainsFolders)
+	{
+		AddAsyncMenuEntry(Section,
+			"SCCMerge",
+			LOCTEXT("SCCMerge", "Merge"),
+			LOCTEXT("SCCMergeTooltip", "Opens the blueprint editor with the merge tool open."),
+			FSlateIcon(),
+			FUIAction(
+				ExecutionCheck(FExecuteAction::CreateSP(this, &FAssetSourceControlContextMenuState::ExecuteSCCMerge)),
+				FCanExecuteAction::CreateLambda([this]() { return IsActionEnabled(CanExecuteSCCMerge()); })
+			),
+			FIsAsyncProcessingActive::CreateLambda([this]() { return IsStillScanning(CanExecuteSCCMerge()); })
+		);
+	}
+
+	Section.AddMenuEntry(
+		"SCCRefresh",
+		LOCTEXT("SCCRefresh", "Refresh"),
+		LOCTEXT("SCCRefreshTooltip", "Updates the source control status of the asset."),
+		FSlateIcon(FAppStyle::GetAppStyleSetName(), "SourceControl.Actions.Refresh"),
+		FUIAction(
+			ExecutionCheck(FExecuteAction::CreateSP(this, &FAssetSourceControlContextMenuState::ExecuteSCCRefresh)),
+			FCanExecuteAction()
+		)
+	);
+}
+
+FExecuteAction FAssetSourceControlContextMenuState::ExecutionCheck(FExecuteAction&& InAction) const
+{
+	return FExecuteAction::CreateLambda([this, Action = MoveTemp(InAction)]()
+		{
+			if (SelectedAssets.Num() > 10) // Todo: Make this into a user option, or different values per operation type
+			{
+				FText Message = FText::Format(LOCTEXT("SCCLargeOperationWarningMessage", "You are about to perform this operation on a large amount of files ({0}), are you sure you want to continue?\n\nUnreal Editor may become unresponsive."), SelectedAssets.Num());
+				FText Title = LOCTEXT("SCCLargeOperationWarningTitle", "Continue Operation?");
+				EAppReturnType::Type Result = FMessageDialog::Open(EAppMsgType::YesNo, Message, &Title);
+
+				if (Result != EAppReturnType::Yes)
+				{
+					return;
+				}
+			}
+
+			FScopedSlowTask SlowTask(0, LOCTEXT("SCCOperationSlowTaskLabel", "Performing Source Control Operation"));
+			SlowTask.MakeDialogDelayed(0.25f);
+			Action.Execute();
+		});
+}
+
+void FAssetSourceControlContextMenuState::ExecuteDiffSelected() const
+{
+	if (SelectedAssets.Num() == 2)
+	{
+		UObject* FirstObjectSelected = SelectedAssets[0].GetAsset();
+		UObject* SecondObjectSelected = SelectedAssets[1].GetAsset();
+
+		if ((FirstObjectSelected != NULL) && (SecondObjectSelected != NULL))
+		{
+			// Load the asset registry module
+			FAssetToolsModule& AssetToolsModule = FModuleManager::GetModuleChecked<FAssetToolsModule>("AssetTools");
+
+			FRevisionInfo CurrentRevision;
+			CurrentRevision.Revision = TEXT("");
+
+			AssetToolsModule.Get().DiffAssets(FirstObjectSelected, SecondObjectSelected, CurrentRevision, CurrentRevision);
+		}
+	}
+}
+
+void FAssetSourceControlContextMenuState::ExecuteSCCMerge() const
+{
+	FAssetToolsModule& AssetToolsModule = FModuleManager::GetModuleChecked<FAssetToolsModule>("AssetTools");
+
+	for (const FAssetData& AssetData : SelectedAssets)
+	{
+		if (UObject* CurrentObject = AssetData.GetAsset())
+		{
+			const FString PackagePath = AssetData.PackageName.ToString();
+			const FString PackageName = AssetData.AssetName.ToString();
+			auto AssetTypeActions = AssetToolsModule.Get().GetAssetTypeActionsForClass(CurrentObject->GetClass()).Pin();
+			if (AssetTypeActions.IsValid())
+			{
+				AssetTypeActions->Merge(CurrentObject);
+			}
+		}
+	}
+}
+
+void FAssetSourceControlContextMenuState::ExecuteSCCCheckOut() const
+{
+	TArray<FString> PackageNames;
+	GetSelectedPackageNames(PackageNames);
+	FEditorFileUtils::CheckoutPackages(PackageNames);
+}
+
+void FAssetSourceControlContextMenuState::ExecuteSCCSyncAndCheckOut() const
+{
+	ExecuteSCCSync();
+	ExecuteSCCCheckOut();
+}
+
+void FAssetSourceControlContextMenuState::ExecuteSCCMakeWritable() const
+{
+	IPlatformFile& PlatformFile = FPlatformFileManager::Get().GetPlatformFile();
+
+	TArray<FString> PackageNames;
+	GetSelectedPackageNames(PackageNames);
+
+	for (const FString& Package : PackageNames)
+	{
+		PlatformFile.SetReadOnly(*SourceControlHelpers::PackageFilename(Package), false);
+	}
+}
+
+void FAssetSourceControlContextMenuState::ExecuteSCCOpenForAdd() const
+{
+	TArray<FString> PackageNames;
+	GetSelectedPackageNames(PackageNames);
+
+	ISourceControlProvider& SourceControlProvider = ISourceControlModule::Get().GetProvider();
+
+	TArray<FString> PackagesToAdd;
+	TArray<UPackage*> PackagesToSave;
+	for (const FString& PackageName : PackageNames)
+	{
+		FSourceControlStatePtr SourceControlState = SourceControlProvider.GetState(SourceControlHelpers::PackageFilename(PackageName), EStateCacheUsage::Use);
+		if (SourceControlState.IsValid() && !SourceControlState->IsSourceControlled())
+		{
+			PackagesToAdd.Add(PackageName);
+
+			// Make sure the file actually exists on disk before adding it
+			FString Filename;
+			if (!FPackageName::DoesPackageExist(PackageName, &Filename))
+			{
+				if (UPackage* Package = FindPackage(NULL, *PackageName))
+				{
+					PackagesToSave.Add(Package);
+				}
+			}
+		}
+	}
+
+	if (PackagesToAdd.Num() > 0)
+	{
+		// If any of the packages are new, save them now
+		if (PackagesToSave.Num() > 0)
+		{
+			const bool bCheckDirty = false;
+			const bool bPromptToSave = false;
+			TArray<UPackage*> FailedPackages;
+			const FEditorFileUtils::EPromptReturnCode Return = FEditorFileUtils::PromptForCheckoutAndSave(PackagesToSave, bCheckDirty, bPromptToSave, &FailedPackages);
+			if (FailedPackages.Num() > 0)
+			{
+				// don't try and add files that failed to save - remove them from the list
+				for (UPackage* FailedPackage : FailedPackages)
+				{
+					PackagesToAdd.Remove(FailedPackage->GetName());
+				}
+			}
+		}
+
+		SourceControlProvider.Execute(ISourceControlOperation::Create<FMarkForAdd>(), SourceControlHelpers::PackageFilenames(PackagesToAdd));
+	}
+}
+
+void FAssetSourceControlContextMenuState::ExecuteSCCCheckIn() const
+{
+	TArray<UPackage*> Packages;
+	GetSelectedPackages(Packages);
+
+	// Prompt the user to ask if they would like to first save any dirty packages they are trying to check-in
+	const FEditorFileUtils::EPromptReturnCode UserResponse = FEditorFileUtils::PromptForCheckoutAndSave(Packages, true, true);
+
+	// If the user elected to save dirty packages, but one or more of the packages failed to save properly OR if the user
+	// canceled out of the prompt, don't follow through on the check-in process
+	const bool bShouldProceed = (UserResponse == FEditorFileUtils::EPromptReturnCode::PR_Success || UserResponse == FEditorFileUtils::EPromptReturnCode::PR_Declined);
+	if (bShouldProceed)
+	{
+		TArray<FString> PackageNames;
+		GetSelectedPackageNames(PackageNames);
+
+		const bool bUseSourceControlStateCache = true;
+
+		FCheckinResultInfo ResultInfo;
+		FSourceControlWindows::PromptForCheckin(ResultInfo, PackageNames, TArray<FString>(), TArray<FString>(), bUseSourceControlStateCache);
+
+		if (ResultInfo.Result == ECommandResult::Failed)
+		{
+			FMessageDialog::Open(EAppMsgType::Ok, NSLOCTEXT("UnrealEd", "SCC_Checkin_Failed", "Check-in failed as a result of save failure."));
+		}
+	}
+	else
+	{
+		// If a failure occurred, alert the user that the check-in was aborted. This warning shouldn't be necessary if the user cancelled
+		// from the dialog, because they obviously intended to cancel the whole operation.
+		if (UserResponse == FEditorFileUtils::EPromptReturnCode::PR_Failure)
+		{
+			FMessageDialog::Open(EAppMsgType::Ok, NSLOCTEXT("UnrealEd", "SCC_Checkin_Aborted", "Check-in aborted as a result of save failure."));
+		}
+	}
+}
+
+void FAssetSourceControlContextMenuState::ExecuteSCCHistory() const
+{
+	TArray<FString> PackageNames;
+	GetSelectedPackageNames(PackageNames);
+	FSourceControlWindows::DisplayRevisionHistory(SourceControlHelpers::PackageFilenames(PackageNames));
+}
+
+void FAssetSourceControlContextMenuState::ExecuteSCCDiffAgainstDepot() const
+{
+	FAssetToolsModule& AssetToolsModule = FModuleManager::GetModuleChecked<FAssetToolsModule>("AssetTools");
+
+	for (const FAssetData& AssetData : SelectedAssets)
+	{
+		if (UObject* CurrentObject = AssetData.GetAsset())
+		{
+			const FString PackagePath = AssetData.PackageName.ToString();
+			const FString PackageName = AssetData.AssetName.ToString();
+			AssetToolsModule.Get().DiffAgainstDepot(CurrentObject, PackagePath, PackageName);
+		}
+	}
+}
+
+void FAssetSourceControlContextMenuState::ExecuteSCCRevert() const
+{
+	TArray<FString> PackageNames;
+	GetSelectedPackageNames(PackageNames);
+	FSourceControlWindows::PromptForRevert(PackageNames);
+}
+
+void FAssetSourceControlContextMenuState::ExecuteSCCRevertWritable() const
+{
+	ISourceControlProvider& SourceControlProvider = ISourceControlModule::Get().GetProvider();
+	IPlatformFile& PlatformFile = FPlatformFileManager::Get().GetPlatformFile();
+
+	TArray<FString> PackageNames;
+	GetSelectedPackageNames(PackageNames);
+
+	TArray<UPackage*> PackagesToReload;
+	TArray<FString> PackageFileNamesToRevert;
+
+	// Only add packages that can actually be reverted
+	for (const FString& PackageName : PackageNames)
+	{
+		FString PackageFileName = SourceControlHelpers::PackageFilename(PackageName);
+
+		if (!PlatformFile.IsReadOnly(*PackageFileName))
+		{
+			FSourceControlStatePtr SourceControlState = SourceControlProvider.GetState(PackageFileName, EStateCacheUsage::Use);
+			if (SourceControlState.IsValid() && SourceControlState->IsSourceControlled() && !SourceControlState->IsCheckedOut())
+			{
+				if (UPackage* Package = FindPackage(nullptr, *PackageName))
+				{
+					PackagesToReload.Add(Package);
+
+					// Detach the linkers of any loaded packages so that SCC can overwrite the files...
+					if (!Package->IsFullyLoaded())
+					{
+						FlushAsyncLoading();
+						Package->FullyLoad();
+					}
+					ResetLoaders(Package);
+				}
+
+				PackageFileNamesToRevert.Add(PackageFileName);
+			}
+		}
+	}
+
+	{
+		TSharedRef<FUpdateStatus, ESPMode::ThreadSafe> UpdateOperation = ISourceControlOperation::Create<FUpdateStatus>();
+		UpdateOperation->SetUpdateHistory(true);
+		SourceControlProvider.Execute(UpdateOperation, PackageFileNamesToRevert);
+	}
+
+	for (const FString& PackageFileName : PackageFileNamesToRevert)
+	{
+		FSourceControlStatePtr SourceControlState = SourceControlProvider.GetState(PackageFileName, EStateCacheUsage::Use);
+		TSharedPtr<class ISourceControlRevision, ESPMode::ThreadSafe> Revision = SourceControlState->GetCurrentRevision();
+
+		if (Revision.IsValid())
+		{
+			TSharedRef<FSync, ESPMode::ThreadSafe> SyncOperation = ISourceControlOperation::Create<FSync>();
+			SyncOperation->SetRevision(Revision->GetRevision());
+			SyncOperation->SetForce(true);
+			if (SourceControlProvider.Execute(SyncOperation, PackageFileName) == ECommandResult::Succeeded)
+			{
+				PlatformFile.SetReadOnly(*PackageFileName, true);
+			}
+		}
+	}
+
+	// Hot-reload the new packages...
+	UPackageTools::ReloadPackages(PackagesToReload);
+}
+
+void FAssetSourceControlContextMenuState::ExecuteSCCSync() const
+{
+	TArray<FString> PackageNames;
+	GetSelectedPackageNames(PackageNames);
+	AssetViewUtils::SyncPackagesFromSourceControl(PackageNames);
+}
+
+bool FAssetSourceControlContextMenuState::CanExecuteSCCMerge() const
+{
+	FAssetToolsModule& AssetToolsModule = FModuleManager::GetModuleChecked<FAssetToolsModule>("AssetTools");
+
+	bool bCanExecuteMerge = bCanExecuteSCCMerge;
+	for (const FAssetData& AssetData : SelectedAssets)
+	{
+		if (UObject* CurrentObject = AssetData.GetAsset())
+		{
+			auto AssetTypeActions = AssetToolsModule.Get().GetAssetTypeActionsForClass(CurrentObject->GetClass()).Pin();
+			if (AssetTypeActions.IsValid())
+			{
+				bCanExecuteMerge = AssetTypeActions->CanMerge();
+			}
+		}
+		else
+		{
+			bCanExecuteMerge = false;
+		}
+	}
+
+	return bCanExecuteMerge;
+}
+
+void FAssetSourceControlContextMenuState::ExecuteSCCRefresh() const
+{
+	TArray<FString> PackageNames;
+	GetSelectedPackageNames(PackageNames);
+
+	ISourceControlModule::Get().GetProvider().Execute(ISourceControlOperation::Create<FUpdateStatus>(), SourceControlHelpers::PackageFilenames(PackageNames), EConcurrency::Asynchronous);
+}
+
+bool FAssetSourceControlContextMenuState::CanExecuteSCCCheckOut() const
+{
+	return bCanExecuteSCCCheckOut;
+}
+
+bool FAssetSourceControlContextMenuState::CanExecuteSCCSyncAndCheckOut() const
+{
+	return bCanExecuteSCCSyncAndCheckOut;
+}
+
+bool FAssetSourceControlContextMenuState::CanExecuteSCCMakeWritable() const
+{
+	return bCanExecuteSCCMakeWritable;
+}
+
+bool FAssetSourceControlContextMenuState::CanExecuteSCCOpenForAdd() const
+{
+	return bCanExecuteSCCOpenForAdd;
+}
+
+bool FAssetSourceControlContextMenuState::CanExecuteSCCCheckIn() const
+{
+	return bCanExecuteSCCCheckIn;
+}
+
+bool FAssetSourceControlContextMenuState::CanExecuteSCCHistory() const
+{
+	return bCanExecuteSCCHistory;
+}
+
+bool FAssetSourceControlContextMenuState::CanExecuteSCCDiffAgainstDepot() const
+{
+	return bCanExecuteSCCHistory;
+}
+
+bool FAssetSourceControlContextMenuState::CanExecuteSCCRevert() const
+{
+	return bCanExecuteSCCRevert;
+}
+
+bool FAssetSourceControlContextMenuState::CanExecuteSCCRevertWritable() const
+{
+	return bCanExecuteSCCRevertWritable;
+}
+
+bool FAssetSourceControlContextMenuState::CanExecuteSCCSync() const
+{
+	return bCanExecuteSCCSync;
+}
+
+bool FAssetSourceControlContextMenuState::CanExecuteDiffSelected() const
+{
+	bool bCanDiffSelected = false;
+	if (SelectedAssets.Num() == 2 && !bContainsFolders)
+	{
+		const FAssetData& FirstSelection = SelectedAssets[0];
+		const FAssetData& SecondSelection = SelectedAssets[1];
+
+		bCanDiffSelected = FirstSelection.AssetClassPath == SecondSelection.AssetClassPath;
+	}
+
+	return bCanDiffSelected;
+}
+
+bool FAssetSourceControlContextMenuState::IsActionEnabled(bool bEnabled) const
+{
+	return bEnabled && !IsStillScanning(bEnabled);
+}
+
+bool FAssetSourceControlContextMenuState::IsStillScanning(bool bEnabled) const
+{
+	// If we only have files selected then we only care about AsyncState
+	// If at least one folder is selected then we will consider done if at least
+	// one file was found that supports the given operation.
+	return AsyncState != EAsyncState::Idle && (bContainsFolders ? !bEnabled : true);
+}
+
+bool FAssetSourceControlContextMenuState::TrySetAsyncState(EAsyncState FromState, EAsyncState ToState)
+{
+	return AsyncState.compare_exchange_strong(FromState, ToState);
+}
+
+void FAssetSourceControlContextMenuState::NextAsyncState()
+{
+	switch (AsyncState)
+	{
+	case EAsyncState::None:
+		if (TrySetAsyncState(EAsyncState::None, EAsyncState::TryCacheCanExecute))
+		{
+			Async(EAsyncExecution::ThreadPool, [this, This = AsShared()]()
+				{
+					TArray<FString> SelectedPaths;
+					Algo::TransformIf(SelectedAssets, SelectedPaths,
+						[](const FAssetData& AssetData) { return AssetData.IsValid(); },
+						[](const FAssetData& AssetData) { return SourceControlHelpers::PackageFilename(AssetData.PackageName.ToString()); });
+
+					TryCacheCanExecuteVars(SelectedPaths, &PathsWithUnknownState);
+
+					// Source control operations have to be started from the main thread
+					Async(EAsyncExecution::TaskGraphMainThread, [this, This = AsShared()]() { NextAsyncState(); });
+				});
+		}
+		break;
+	case EAsyncState::TryCacheCanExecute:
+		if (PathsWithUnknownState.Num() == 0)
+		{
+			// Skip to the end since there's no more work to do
+			TrySetAsyncState(EAsyncState::TryCacheCanExecute, EAsyncState::Idle);
+		}
+		else if (TrySetAsyncState(EAsyncState::TryCacheCanExecute, EAsyncState::SCCUpdate))
+		{
+			// Do an async force update to make sure to get mode up-to-date status
+			FSourceControlOperationRef Operation = ISourceControlOperation::Create<FUpdateStatus>();
+			SCCOperation = Operation;
+
+			ECommandResult::Type Result = ISourceControlModule::Get().GetProvider().Execute(Operation, PathsWithUnknownState, EConcurrency::Asynchronous,
+				FSourceControlOperationComplete::CreateSP(this, &FAssetSourceControlContextMenuState::CacheCanExecuteVarsSCCHandler));
+
+			if (Result != ECommandResult::Succeeded)
+			{
+				AsyncState = EAsyncState::Error;
+			}
+		}
+		break;
+	case EAsyncState::SCCUpdate:
+		TrySetAsyncState(EAsyncState::SCCUpdate, EAsyncState::Idle);
+		break;
+	}
+}
+
+void FAssetSourceControlContextMenuState::CacheCanExecuteVarsSCCHandler(const FSourceControlOperationRef& InOperation, ECommandResult::Type Result)
+{
+	EAsyncState ExpectedState = EAsyncState::SCCUpdate;
+	if (Result == ECommandResult::Succeeded)
+	{
+		TryCacheCanExecuteVars(PathsWithUnknownState, nullptr);
+	}
+
+	NextAsyncState();
+}
+
+void FAssetSourceControlContextMenuState::TryCacheCanExecuteVars(const TArray<FString>& InSelectedPaths, TArray<FString>* OptionalPathsWithUnknownState)
+{
+	ISourceControlProvider& SourceControlProvider = ISourceControlModule::Get().GetProvider();
+	IPlatformFile& PlatformFile = FPlatformFileManager::Get().GetPlatformFile();
+
+	const bool bUsesCheckout = SourceControlProvider.UsesCheckout();
+	const bool bUsesFileRevisions = SourceControlProvider.UsesFileRevisions();
+
+	// Check the SCC state for each package in the selected paths
+	TArray<FSourceControlStateRef> SelectedStates;
+	if (SourceControlProvider.GetState(InSelectedPaths, SelectedStates, EStateCacheUsage::Use) == ECommandResult::Succeeded)
+	{
+		for (const FSourceControlStateRef& SourceControlState : SelectedStates)
+		{
+			if (AsyncState == EAsyncState::Cancelled)
+				return;
+
+			if (SourceControlState->IsUnknown())
+			{
+				if (OptionalPathsWithUnknownState != nullptr)
+				{
+					OptionalPathsWithUnknownState->Add(SourceControlState->GetFilename());
+				}
+				continue;
+			}
+
+			if (SourceControlState->IsConflicted())
+			{
+				bCanExecuteSCCMerge = true;
+			}
+
+			bool bCanSyncCurrentItem = false;
+			const bool bIsReadOnly = PlatformFile.IsReadOnly(*SourceControlState->GetFilename());
+
+			if (bUsesFileRevisions)
+			{
+				if (!SourceControlState->IsCurrent())
+				{
+					bCanExecuteSCCSync = bCanSyncCurrentItem = true;
+				}
+
+				if (SourceControlState->CanCheckIn())
+				{
+					bCanExecuteSCCCheckIn = true;
+				}
+
+				if (SourceControlState->IsSourceControlled() && !SourceControlState->IsCheckedOut() && !bIsReadOnly)
+				{
+					bCanExecuteSCCRevertWritable = true;
+				}
+			}
+
+			if (bUsesCheckout)
+			{
+				if (SourceControlState->CanCheckout())
+				{
+					bCanExecuteSCCCheckOut = true;
+				}
+				else
+				{
+					FString CheckedOutOther;
+					if (!bContainsFolders && SourceControlState->IsCheckedOutOther(&CheckedOutOther))
+					{
+						CheckedOutUsers.Add(FString::Printf(TEXT("%s checked out by %s"), *SourceControlState->GetFilename(), *CheckedOutOther));
+					}
+					else if (bCanSyncCurrentItem)
+					{
+						bCanExecuteSCCSyncAndCheckOut = true;
+					}
+				}
+
+				if (bIsReadOnly)
+				{
+					bCanExecuteSCCMakeWritable = true;
+				}
+			}
+
+			if (!SourceControlState->IsSourceControlled())
+			{
+				bCanExecuteSCCOpenForAdd = SourceControlState->CanAdd();
+			}
+			else
+			{
+				bCanExecuteSCCHistory = !SourceControlState->IsAdded();
+			}
+
+			if (SourceControlState->CanRevert())
+			{
+				bCanExecuteSCCRevert = true;
+			}
+		}
+	}
+
+	CheckedOutUsersText = FText::FromString(FString::Join(CheckedOutUsers, TEXT("\n")));
+}
+
+void FAssetSourceControlContextMenuState::GetSelectedPackageNames(TArray<FString>& OutPackageNames) const
+{
+	for (const FAssetData& AssetData : SelectedAssets)
+	{
+		OutPackageNames.Add(AssetData.PackageName.ToString());
+	}
+}
+
+void FAssetSourceControlContextMenuState::GetSelectedPackages(TArray<UPackage*>& OutPackages) const
+{
+	for (const FAssetData& AssetData : SelectedAssets)
+	{
+		if (UPackage* Package = FindPackage(NULL, *AssetData.PackageName.ToString()))
+		{
+			OutPackages.Add(Package);
+		}
+	}
+}
+
+#undef LOCTEXT_NAMESPACE
diff --git a/Engine/Plugins/Editor/AssetManagerEditor/Source/AssetManagerEditor/Private/AssetSourceControlContextMenu.h b/Engine/Plugins/Editor/AssetManagerEditor/Source/AssetManagerEditor/Private/AssetSourceControlContextMenu.h
new file mode 100644
--- /dev/null
+++ b/Engine/Plugins/Editor/AssetManagerEditor/Source/AssetManagerEditor/Private/AssetSourceControlContextMenu.h
@@ -0,0 +1,25 @@
+// Copyright Epic Games, Inc. All Rights Reserved.
+
+#pragma once
+
+#include "CoreMinimal.h"
+#include "AssetRegistry/AssetData.h"
+#include "ISourceControlProvider.h"
+#include "HAL/CriticalSection.h"
+
+class UToolMenu;
+struct FToolMenuSection;
+class FAssetSourceControlContextMenuState;
+
+// Display source control options in the given menu section based on what's selected in the menu context
+class FAssetSourceControlContextMenu
+{
+public:
+	~FAssetSourceControlContextMenu();
+
+	/** Makes the context menu widget */
+	void MakeContextMenu(FToolMenuSection& InSection);
+
+private:
+	TSharedPtr<FAssetSourceControlContextMenuState> InnerState;
+};
diff --git a/Engine/Plugins/Editor/ContentBrowser/ContentBrowserAssetDataSource/Source/ContentBrowserAssetDataSource/Private/AssetContextMenuUtils.cpp b/Engine/Plugins/Editor/ContentBrowser/ContentBrowserAssetDataSource/Source/ContentBrowserAssetDataSource/Private/AssetContextMenuUtils.cpp
deleted file mode 100644
diff --git a/Engine/Plugins/Editor/ContentBrowser/ContentBrowserAssetDataSource/Source/ContentBrowserAssetDataSource/Private/AssetContextMenuUtils.h b/Engine/Plugins/Editor/ContentBrowser/ContentBrowserAssetDataSource/Source/ContentBrowserAssetDataSource/Private/AssetContextMenuUtils.h
deleted file mode 100644
diff --git a/Engine/Plugins/Editor/ContentBrowser/ContentBrowserAssetDataSource/Source/ContentBrowserAssetDataSource/Private/AssetFileContextMenu.cpp b/Engine/Plugins/Editor/ContentBrowser/ContentBrowserAssetDataSource/Source/ContentBrowserAssetDataSource/Private/AssetFileContextMenu.cpp
--- a/Engine/Plugins/Editor/ContentBrowser/ContentBrowserAssetDataSource/Source/ContentBrowserAssetDataSource/Private/AssetFileContextMenu.cpp
+++ b/Engine/Plugins/Editor/ContentBrowser/ContentBrowserAssetDataSource/Source/ContentBrowserAssetDataSource/Private/AssetFileContextMenu.cpp
@@ -22,6 +22,7 @@
 #include "Framework/MultiBox/MultiBoxBuilder.h"
 #include "ToolMenus.h"
 #include "ContentBrowserMenuContexts.h"
+#include "ContentBrowserDataMenuContexts.h"
 #include "ContentBrowserDataSource.h"
 #include "Widgets/Input/SButton.h"
 #include "Styling/AppStyle.h"
@@ -31,9 +32,6 @@
 #include "UnrealClient.h"
 #include "Materials/MaterialFunctionInstance.h"
 #include "Materials/Material.h"
-#include "SourceControlOperations.h"
-#include "ISourceControlModule.h"
-#include "SourceControlHelpers.h"
 #include "Settings/EditorExperimentalSettings.h"
 #include "Materials/MaterialInstanceConstant.h"
 #include "FileHelpers.h"
@@ -56,7 +54,6 @@
 #include "Internationalization/PackageLocalizationUtil.h"
 #include "Internationalization/TextLocalizationResource.h"
 
-#include "SourceControlWindows.h"
 #include "Kismet2/KismetEditorUtilities.h"
 #include "ComponentAssetBroker.h"
 #include "Widgets/Input/SNumericEntryBox.h"
@@ -106,8 +103,6 @@ void FAssetFileContextMenu::AddMenuOptions(UToolMenu* Menu)
 
 	ParentWidget = Context->ParentWidget;
 
-	OnRefreshView = Context->OnRefreshView;
-
 	// Cache any vars that are used in determining if you can execute any actions.
 	// Useful for actions whose "CanExecute" will not change or is expensive to calculate.
 	CacheCanExecuteVars();
@@ -123,12 +118,6 @@ void FAssetFileContextMenu::AddMenuOptions(UToolMenu* Menu)
 
 	// Add documentation options
 	AddDocumentationMenuOptions(Menu);
-
-	// Add source control options
-	if (Context->bCanBeModified)
-	{
-		AddSourceControlMenuOptions(Menu);
-	}
 }
 
 bool FAssetFileContextMenu::AddImportedAssetMenuOptions(UToolMenu* Menu)
@@ -1160,188 +1149,6 @@ bool FAssetFileContextMenu::AddDocumentationMenuOptions(UToolMenu* Menu)
 	return bAddedOption;
 }
 
-bool FAssetFileContextMenu::AddSourceControlMenuOptions(UToolMenu* Menu)
-{
-	FToolMenuSection& Section = Menu->AddSection("AssetContextSourceControl");
-	
-	if ( ISourceControlModule::Get().IsEnabled() )
-	{
-		// SCC sub menu
-		Section.AddSubMenu(
-			"SourceControlSubMenu",
-			LOCTEXT("SourceControlSubMenuLabel", "Source Control"),
-			LOCTEXT("SourceControlSubMenuToolTip", "Source control actions."),
-			FNewToolMenuDelegate::CreateSP(this, &FAssetFileContextMenu::FillSourceControlSubMenu),
-			FUIAction(
-				FExecuteAction(),
-				FCanExecuteAction::CreateSP( this, &FAssetFileContextMenu::CanExecuteSourceControlActions )
-				),
-			EUserInterfaceActionType::Button,
-			false,
-			FSlateIcon(FAppStyle::GetAppStyleSetName(), "SourceControl.StatusIcon.On")
-			);
-	}
-	else
-	{
-		Section.AddMenuEntry(
-			"SCCConnectToSourceControl",
-			LOCTEXT("SCCConnectToSourceControl", "Connect To Source Control..."),
-			LOCTEXT("SCCConnectToSourceControlTooltip", "Connect to source control to allow source control operations to be performed on content and levels."),
-			FSlateIcon(FAppStyle::GetAppStyleSetName(), "MainFrame.ConnectToSourceControl"),
-			FUIAction(
-				FExecuteAction::CreateSP( this, &FAssetFileContextMenu::ExecuteEnableSourceControl ),
-				FCanExecuteAction::CreateSP( this, &FAssetFileContextMenu::CanExecuteSourceControlActions )
-				)
-			);
-	}
-
-	// Diff selected
-	if (CanExecuteDiffSelected())
-	{
-		Section.AddMenuEntry(
-			"DiffSelected",
-			LOCTEXT("DiffSelected", "Diff Selected"),
-			LOCTEXT("DiffSelectedTooltip", "Diff the two assets that you have selected."),
-			FSlateIcon(FAppStyle::GetAppStyleSetName(), "SourceControl.Actions.Diff"),
-			FUIAction(
-				FExecuteAction::CreateSP(this, &FAssetFileContextMenu::ExecuteDiffSelected)
-			)
-		);
-	}
-
-	return true;
-}
-
-void FAssetFileContextMenu::FillSourceControlSubMenu(UToolMenu* Menu)
-{
-	FToolMenuSection& Section = Menu->AddSection("AssetSourceControlActions", LOCTEXT("AssetSourceControlActionsMenuHeading", "Source Control"));
-
-	if( CanExecuteSCCMerge() )
-	{
-		Section.AddMenuEntry(
-			"SCCMerge",
-			LOCTEXT("SCCMerge", "Merge"),
-			LOCTEXT("SCCMergeTooltip", "Opens the blueprint editor with the merge tool open."),
-			FSlateIcon(),
-			FUIAction(
-				FExecuteAction::CreateSP(this, &FAssetFileContextMenu::ExecuteSCCMerge),
-				FCanExecuteAction::CreateSP(this, &FAssetFileContextMenu::CanExecuteSCCMerge)
-			)
-		);
-	}
-
-	if( CanExecuteSCCSync() )
-	{
-		Section.AddMenuEntry(
-			"SCCSync",
-			LOCTEXT("SCCSync", "Sync"),
-			LOCTEXT("SCCSyncTooltip", "Updates the item to the latest version in source control."),
-			FSlateIcon(FAppStyle::GetAppStyleSetName(), "SourceControl.Actions.Sync"),
-			FUIAction(
-				FExecuteAction::CreateSP( this, &FAssetFileContextMenu::ExecuteSCCSync ),
-				FCanExecuteAction::CreateSP( this, &FAssetFileContextMenu::CanExecuteSCCSync )
-			)
-		);
-	}
-
-	if ( CanExecuteSCCCheckOut() )
-	{
-		Section.AddMenuEntry(
-			"SCCCheckOut",
-			LOCTEXT("SCCCheckOut", "Check Out"),
-			LOCTEXT("SCCCheckOutTooltip", "Checks out the selected asset from source control."),
-			FSlateIcon(FAppStyle::GetAppStyleSetName(), "SourceControl.Actions.CheckOut"),
-			FUIAction(
-				FExecuteAction::CreateSP( this, &FAssetFileContextMenu::ExecuteSCCCheckOut ),
-				FCanExecuteAction::CreateSP( this, &FAssetFileContextMenu::CanExecuteSCCCheckOut )
-			)
-		);
-	}
-
-	if ( CanExecuteSCCOpenForAdd() )
-	{
-		Section.AddMenuEntry(
-			"SCCOpenForAdd",
-			LOCTEXT("SCCOpenForAdd", "Mark For Add"),
-			LOCTEXT("SCCOpenForAddTooltip", "Adds the selected asset to source control."),
-			FSlateIcon(FAppStyle::GetAppStyleSetName(), "SourceControl.Actions.Add"),
-			FUIAction(
-				FExecuteAction::CreateSP( this, &FAssetFileContextMenu::ExecuteSCCOpenForAdd ),
-				FCanExecuteAction::CreateSP( this, &FAssetFileContextMenu::CanExecuteSCCOpenForAdd )
-			)
-		);
-	}
-
-	if ( CanExecuteSCCCheckIn() )
-	{
-		Section.AddMenuEntry(
-			"SCCCheckIn",
-			LOCTEXT("SCCCheckIn", "Check In"),
-			LOCTEXT("SCCCheckInTooltip", "Checks in the selected asset to source control."),
-			FSlateIcon(FAppStyle::GetAppStyleSetName(), "SourceControl.Actions.Submit"),
-			FUIAction(
-				FExecuteAction::CreateSP( this, &FAssetFileContextMenu::ExecuteSCCCheckIn ),
-				FCanExecuteAction::CreateSP( this, &FAssetFileContextMenu::CanExecuteSCCCheckIn )
-			)
-		);
-	}
-
-	Section.AddMenuEntry(
-		"SCCRefresh",
-		LOCTEXT("SCCRefresh", "Refresh"),
-		LOCTEXT("SCCRefreshTooltip", "Updates the source control status of the asset."),
-		FSlateIcon(FAppStyle::GetAppStyleSetName(), "SourceControl.Actions.Refresh"),
-		FUIAction(
-			FExecuteAction::CreateSP( this, &FAssetFileContextMenu::ExecuteSCCRefresh ),
-			FCanExecuteAction::CreateSP( this, &FAssetFileContextMenu::CanExecuteSCCRefresh )
-			)
-		);
-
-	if( CanExecuteSCCHistory() )
-	{
-		Section.AddMenuEntry(
-			"SCCHistory",
-			LOCTEXT("SCCHistory", "History"),
-			LOCTEXT("SCCHistoryTooltip", "Displays the source control revision history of the selected asset."),
-			FSlateIcon(FAppStyle::GetAppStyleSetName(), "SourceControl.Actions.History"),
-			FUIAction(
-				FExecuteAction::CreateSP( this, &FAssetFileContextMenu::ExecuteSCCHistory ),
-				FCanExecuteAction::CreateSP( this, &FAssetFileContextMenu::CanExecuteSCCHistory )
-			)
-		);
-
-		Section.AddMenuEntry(
-			"SCCDiffAgainstDepot",
-			LOCTEXT("SCCDiffAgainstDepot", "Diff Against Depot"),
-			LOCTEXT("SCCDiffAgainstDepotTooltip", "Look at differences between your version of the asset and that in source control."),
-			FSlateIcon(FAppStyle::GetAppStyleSetName(), "SourceControl.Actions.Diff"),
-			FUIAction(
-				FExecuteAction::CreateSP( this, &FAssetFileContextMenu::ExecuteSCCDiffAgainstDepot ),
-				FCanExecuteAction::CreateSP( this, &FAssetFileContextMenu::CanExecuteSCCDiffAgainstDepot )
-			)
-		);	
-	}
-
-	if( CanExecuteSCCRevert() )
-	{
-		Section.AddMenuEntry(
-			"SCCRevert",
-			LOCTEXT("SCCRevert", "Revert"),
-			LOCTEXT("SCCRevertTooltip", "Reverts the asset to the state it was before it was checked out."),
-			FSlateIcon(FAppStyle::GetAppStyleSetName(), "SourceControl.Actions.Revert"),
-			FUIAction(
-				FExecuteAction::CreateSP( this, &FAssetFileContextMenu::ExecuteSCCRevert ),
-				FCanExecuteAction::CreateSP( this, &FAssetFileContextMenu::CanExecuteSCCRevert )
-			)
-		);
-	}
-}
-
-bool FAssetFileContextMenu::CanExecuteSourceControlActions() const
-{
-	return SelectedAssets.Num() > 0;
-}
-
 bool FAssetFileContextMenu::AreImportedAssetActionsVisible() const
 {
 	FAssetToolsModule& AssetToolsModule = FModuleManager::Get().LoadModuleChecked<FAssetToolsModule>("AssetTools");
@@ -1532,11 +1339,11 @@ void FAssetFileContextMenu::GetSelectedAssets(TArray<UObject*>& Assets, bool Ski
 	for (const FAssetData& SelectedAsset : SelectedAssets)
 	{
 		if (SkipRedirectors && (SelectedAsset.AssetClassPath == UObjectRedirector::StaticClass()->GetClassPathName()))
-    	{
-    		// Don't operate on Redirectors
-    		continue;
-    	}
-	
+		{
+			// Don't operate on Redirectors
+			continue;
+		}
+
 		SelectedAssetPaths.Add(SelectedAsset.GetObjectPathString());
 	}
 
@@ -2037,184 +1844,6 @@ void FAssetFileContextMenu::ExecuteBulkExport()
 	}
 }
 
-void FAssetFileContextMenu::ExecuteSCCRefresh()
-{
-	TArray<FString> PackageNames;
-	GetSelectedPackageNames(PackageNames);
-
-	ISourceControlModule::Get().GetProvider().Execute(ISourceControlOperation::Create<FUpdateStatus>(), SourceControlHelpers::PackageFilenames(PackageNames), EConcurrency::Asynchronous);
-}
-
-void FAssetFileContextMenu::ExecuteSCCMerge()
-{
-	FAssetToolsModule& AssetToolsModule = FModuleManager::GetModuleChecked<FAssetToolsModule>("AssetTools");
-
-	for (int32 AssetIdx = 0; AssetIdx < SelectedAssets.Num(); AssetIdx++)
-	{
-		// Get the actual asset (will load it)
-		const FAssetData& AssetData = SelectedAssets[AssetIdx];
-
-		UObject* CurrentObject = AssetData.GetAsset();
-		if (CurrentObject)
-		{
-			const FString PackagePath = AssetData.PackageName.ToString();
-			const FString PackageName = AssetData.AssetName.ToString();
-			auto AssetTypeActions = AssetToolsModule.Get().GetAssetTypeActionsForClass( CurrentObject->GetClass() ).Pin();
-			if( AssetTypeActions.IsValid() )
-			{
-				AssetTypeActions->Merge(CurrentObject);
-			}
-		}
-	}
-}
-
-void FAssetFileContextMenu::ExecuteSCCCheckOut()
-{
-	TArray<UPackage*> PackagesToCheckOut;
-	GetSelectedPackages(PackagesToCheckOut);
-
-	if ( PackagesToCheckOut.Num() > 0 )
-	{
-		FEditorFileUtils::CheckoutPackages(PackagesToCheckOut);
-	}
-}
-
-void FAssetFileContextMenu::ExecuteSCCOpenForAdd()
-{
-	TArray<FString> PackageNames;
-	GetSelectedPackageNames(PackageNames);
-
-	ISourceControlProvider& SourceControlProvider = ISourceControlModule::Get().GetProvider();
-
-	TArray<FString> PackagesToAdd;
-	TArray<UPackage*> PackagesToSave;
-	for ( auto PackageIt = PackageNames.CreateConstIterator(); PackageIt; ++PackageIt )
-	{
-		FSourceControlStatePtr SourceControlState = SourceControlProvider.GetState(SourceControlHelpers::PackageFilename(*PackageIt), EStateCacheUsage::Use);
-		if ( SourceControlState.IsValid() && !SourceControlState->IsSourceControlled() )
-		{
-			PackagesToAdd.Add(*PackageIt);
-
-			// Make sure the file actually exists on disk before adding it
-			FString Filename;
-			if ( !FPackageName::DoesPackageExist(*PackageIt, &Filename) )
-			{
-				UPackage* Package = FindPackage(NULL, **PackageIt);
-				if ( Package )
-				{
-					PackagesToSave.Add(Package);
-				}
-			}
-		}
-	}
-
-	if ( PackagesToAdd.Num() > 0 )
-	{
-		// If any of the packages are new, save them now
-		if ( PackagesToSave.Num() > 0 )
-		{
-			const bool bCheckDirty = false;
-			const bool bPromptToSave = false;
-			TArray<UPackage*> FailedPackages;
-			const FEditorFileUtils::EPromptReturnCode Return = FEditorFileUtils::PromptForCheckoutAndSave(PackagesToSave, bCheckDirty, bPromptToSave, &FailedPackages);
-			if(FailedPackages.Num() > 0)
-			{
-				// don't try and add files that failed to save - remove them from the list
-				for(auto FailedPackageIt = FailedPackages.CreateConstIterator(); FailedPackageIt; FailedPackageIt++)
-				{
-					PackagesToAdd.Remove((*FailedPackageIt)->GetName());
-				}
-			}
-		}
-
-		SourceControlProvider.Execute(ISourceControlOperation::Create<FMarkForAdd>(), SourceControlHelpers::PackageFilenames(PackagesToAdd));
-	}
-}
-
-void FAssetFileContextMenu::ExecuteSCCCheckIn()
-{
-	TArray<UPackage*> Packages;
-	GetSelectedPackages(Packages);
-
-	// Prompt the user to ask if they would like to first save any dirty packages they are trying to check-in
-	const FEditorFileUtils::EPromptReturnCode UserResponse = FEditorFileUtils::PromptForCheckoutAndSave( Packages, true, true );
-
-	// If the user elected to save dirty packages, but one or more of the packages failed to save properly OR if the user
-	// canceled out of the prompt, don't follow through on the check-in process
-	const bool bShouldProceed = ( UserResponse == FEditorFileUtils::EPromptReturnCode::PR_Success || UserResponse == FEditorFileUtils::EPromptReturnCode::PR_Declined );
-	if ( bShouldProceed )
-	{
-		TArray<FString> PackageNames;
-		GetSelectedPackageNames(PackageNames);
-
-		const bool bUseSourceControlStateCache = true;
-
-		FCheckinResultInfo ResultInfo;
-		FSourceControlWindows::PromptForCheckin(ResultInfo, PackageNames, TArray<FString>(), TArray<FString>(), bUseSourceControlStateCache);
-
-		if ( ResultInfo.Result == ECommandResult::Failed )
-		{
-			FMessageDialog::Open(EAppMsgType::Ok, NSLOCTEXT("UnrealEd", "SCC_Checkin_Failed", "Check-in failed as a result of save failure."));
-		}
-	}
-	else
-	{
-		// If a failure occurred, alert the user that the check-in was aborted. This warning shouldn't be necessary if the user cancelled
-		// from the dialog, because they obviously intended to cancel the whole operation.
-		if ( UserResponse == FEditorFileUtils::EPromptReturnCode::PR_Failure )
-		{
-			FMessageDialog::Open( EAppMsgType::Ok, NSLOCTEXT("UnrealEd", "SCC_Checkin_Aborted", "Check-in aborted as a result of save failure.") );
-		}
-	}
-}
-
-void FAssetFileContextMenu::ExecuteSCCHistory()
-{
-	TArray<FString> PackageNames;
-	GetSelectedPackageNames(PackageNames);
-	FSourceControlWindows::DisplayRevisionHistory(SourceControlHelpers::PackageFilenames(PackageNames));
-}
-
-void FAssetFileContextMenu::ExecuteSCCDiffAgainstDepot() const
-{
-	// Load the asset registry module
-	FAssetToolsModule& AssetToolsModule = FModuleManager::GetModuleChecked<FAssetToolsModule>("AssetTools");
-
-	// Iterate over each selected asset
-	for(int32 AssetIdx=0; AssetIdx<SelectedAssets.Num(); AssetIdx++)
-	{
-		// Get the actual asset (will load it)
-		const FAssetData& AssetData = SelectedAssets[AssetIdx];
-
-		UObject* CurrentObject = AssetData.GetAsset();
-		if( CurrentObject )
-		{
-			const FString PackagePath = AssetData.PackageName.ToString();
-			const FString PackageName = AssetData.AssetName.ToString();
-			AssetToolsModule.Get().DiffAgainstDepot( CurrentObject, PackagePath, PackageName );
-		}
-	}
-}
-
-void FAssetFileContextMenu::ExecuteSCCRevert()
-{
-	TArray<FString> PackageNames;
-	GetSelectedPackageNames(PackageNames);
-	FSourceControlWindows::PromptForRevert(PackageNames);
-}
-
-void FAssetFileContextMenu::ExecuteSCCSync()
-{
-	TArray<FString> PackageNames;
-	GetSelectedPackageNames(PackageNames);
-	AssetViewUtils::SyncPackagesFromSourceControl(PackageNames);
-}
-
-void FAssetFileContextMenu::ExecuteEnableSourceControl()
-{
-	ISourceControlModule::Get().ShowLoginDialog(FSourceControlLoginClosed(), ELoginWindowMode::Modeless);
-}
-
 bool FAssetFileContextMenu::CanExecuteCreateBlueprintUsing() const
 {
 	// Only work if you have a single asset selected
@@ -2301,87 +1930,6 @@ bool FAssetFileContextMenu::CanExecuteShowAssetMetaData() const
 	return bResult;
 }
 
-bool FAssetFileContextMenu::CanExecuteSCCRefresh() const
-{
-	return ISourceControlModule::Get().IsEnabled();
-}
-
-bool FAssetFileContextMenu::CanExecuteSCCMerge() const
-{
-	FAssetToolsModule& AssetToolsModule = FModuleManager::GetModuleChecked<FAssetToolsModule>("AssetTools");
-
-	bool bCanExecuteMerge = bCanExecuteSCCMerge;
-	for (int32 AssetIdx = 0; AssetIdx < SelectedAssets.Num() && bCanExecuteMerge; AssetIdx++)
-	{
-		// Get the actual asset (will load it)
-		const FAssetData& AssetData = SelectedAssets[AssetIdx];
-		UObject* CurrentObject = AssetData.GetAsset();
-		if (CurrentObject)
-		{
-			auto AssetTypeActions = AssetToolsModule.Get().GetAssetTypeActionsForClass(CurrentObject->GetClass()).Pin();
-			if (AssetTypeActions.IsValid())
-			{
-				bCanExecuteMerge = AssetTypeActions->CanMerge();
-			}
-		}
-		else
-		{
-			bCanExecuteMerge = false;
-		}
-	}
-
-	return bCanExecuteMerge;
-}
-
-bool FAssetFileContextMenu::CanExecuteSCCCheckOut() const
-{
-	return bCanExecuteSCCCheckOut;
-}
-
-bool FAssetFileContextMenu::CanExecuteSCCOpenForAdd() const
-{
-	return bCanExecuteSCCOpenForAdd;
-}
-
-bool FAssetFileContextMenu::CanExecuteSCCCheckIn() const
-{	
-	return bCanExecuteSCCCheckIn && bUsesFileRevisions;
-}
-
-bool FAssetFileContextMenu::CanExecuteSCCHistory() const
-{
-	return bCanExecuteSCCHistory;
-}
-
-bool FAssetFileContextMenu::CanExecuteSCCDiffAgainstDepot() const
-{
-	return bCanExecuteSCCHistory;
-}
-
-bool FAssetFileContextMenu::CanExecuteSCCRevert() const
-{
-	return bCanExecuteSCCRevert;
-}
-
-bool FAssetFileContextMenu::CanExecuteSCCSync() const
-{
-	return bCanExecuteSCCSync;
-}
-
-bool FAssetFileContextMenu::CanExecuteDiffSelected() const
-{
-	bool bCanDiffSelected = false;
-	if (SelectedAssets.Num() == 2)
-	{
-		FAssetData const& FirstSelection = SelectedAssets[0];
-		FAssetData const& SecondSelection = SelectedAssets[1];
-
-		bCanDiffSelected = FirstSelection.AssetClassPath == SecondSelection.AssetClassPath;
-	}
-
-	return bCanDiffSelected;
-}
-
 bool FAssetFileContextMenu::CanExecuteCaptureThumbnail() const
 {
 	return GCurrentLevelEditingViewportClient != NULL;
@@ -2403,14 +1951,6 @@ bool FAssetFileContextMenu::CanClearCustomThumbnails() const
 void FAssetFileContextMenu::CacheCanExecuteVars()
 {
 	bAtLeastOneNonRedirectorSelected = false;
-	bCanExecuteSCCMerge = false;
-	bCanExecuteSCCCheckOut = false;
-	bCanExecuteSCCOpenForAdd = false;
-	bCanExecuteSCCCheckIn = false;
-	bCanExecuteSCCHistory = false;
-	bCanExecuteSCCRevert = false;
-	bCanExecuteSCCSync = false;
-	bUsesFileRevisions = false;
 
 	for (auto AssetIt = SelectedAssets.CreateConstIterator(); AssetIt; ++AssetIt)
 	{
@@ -2425,68 +1965,13 @@ void FAssetFileContextMenu::CacheCanExecuteVars()
 			bAtLeastOneNonRedirectorSelected = true;
 		}
 
-		ISourceControlProvider& SourceControlProvider = ISourceControlModule::Get().GetProvider();
-		if ( ISourceControlModule::Get().IsEnabled() )
-		{
-			// Check the SCC state for each package in the selected paths
-			FSourceControlStatePtr SourceControlState = SourceControlProvider.GetState(SourceControlHelpers::PackageFilename(AssetData.PackageName.ToString()), EStateCacheUsage::Use);
-			if(SourceControlState.IsValid())
-			{
-				if (SourceControlState->IsConflicted() )
-				{
-					bCanExecuteSCCMerge = true;
-				}
-
-				if ( SourceControlState->CanCheckout() )
-				{
-					bCanExecuteSCCCheckOut = true;
-				}
-
-				if ( !SourceControlState->IsSourceControlled() && SourceControlState->CanAdd() )
-				{
-					bCanExecuteSCCOpenForAdd = true;
-				}
-				else if( SourceControlState->IsSourceControlled() && !SourceControlState->IsAdded() )
-				{
-					bCanExecuteSCCHistory = true;
-				}
-
-				if(!SourceControlState->IsCurrent() && SourceControlProvider.UsesFileRevisions())
-				{
-					bCanExecuteSCCSync = true;
-				}
-
-				if ( SourceControlState->CanCheckIn() )
-				{
-					bCanExecuteSCCCheckIn = true;
-				}
-
-				if (SourceControlState->CanRevert())
-				{
-					bCanExecuteSCCRevert = true;
-				}
-			}
-		}
-
 		if ( bAtLeastOneNonRedirectorSelected
-			&& bCanExecuteSCCMerge
-			&& bCanExecuteSCCCheckOut
-			&& bCanExecuteSCCOpenForAdd
-			&& bCanExecuteSCCCheckIn
-			&& bCanExecuteSCCHistory
-			&& bCanExecuteSCCRevert
-			&& bCanExecuteSCCSync
 			)
 		{
 			// All options are available, no need to keep iterating
 			break;
 		}
 	}
-
-	if (ISourceControlModule::Get().IsEnabled() && ISourceControlModule::Get().GetProvider().IsAvailable())
-	{
-		bUsesFileRevisions = ISourceControlModule::Get().GetProvider().UsesChangelists();
-	}
 }
 
 void FAssetFileContextMenu::GetSelectedPackageNames(TArray<FString>& OutPackageNames) const
diff --git a/Engine/Plugins/Editor/ContentBrowser/ContentBrowserAssetDataSource/Source/ContentBrowserAssetDataSource/Private/AssetFolderContextMenu.cpp b/Engine/Plugins/Editor/ContentBrowser/ContentBrowserAssetDataSource/Source/ContentBrowserAssetDataSource/Private/AssetFolderContextMenu.cpp
--- a/Engine/Plugins/Editor/ContentBrowser/ContentBrowserAssetDataSource/Source/ContentBrowserAssetDataSource/Private/AssetFolderContextMenu.cpp
+++ b/Engine/Plugins/Editor/ContentBrowser/ContentBrowserAssetDataSource/Source/ContentBrowserAssetDataSource/Private/AssetFolderContextMenu.cpp
@@ -2,51 +2,16 @@
 
 #include "AssetFolderContextMenu.h"
 #include "Misc/MessageDialog.h"
-#include "HAL/FileManager.h"
-#include "Misc/Paths.h"
 #include "Modules/ModuleManager.h"
 #include "UObject/ObjectRedirector.h"
-#include "Misc/PackageName.h"
-#include "Widgets/DeclarativeSyntaxSupport.h"
-#include "Widgets/SBoxPanel.h"
-#include "Widgets/SWindow.h"
-#include "Framework/Application/SlateApplication.h"
 #include "Textures/SlateIcon.h"
-#include "Framework/MultiBox/MultiBoxExtender.h"
-#include "Widgets/Input/SButton.h"
-#include "Widgets/Colors/SColorBlock.h"
-#include "Styling/AppStyle.h"
-#include "SourceControlOperations.h"
-#include "ISourceControlModule.h"
-#include "SourceControlHelpers.h"
 #include "AssetRegistry/AssetData.h"
-#include "Editor.h"
-#include "FileHelpers.h"
 #include "AssetRegistry/ARFilter.h"
 #include "AssetRegistry/AssetRegistryModule.h"
-#include "IAssetTools.h"
 #include "AssetToolsModule.h"
 #include "AssetViewUtils.h"
-//#include "ContentBrowserLog.h"
-//#include "ContentBrowserSingleton.h"
-//#include "ContentBrowserUtils.h"
-#include "SourceControlWindows.h"
-#include "ContentBrowserModule.h"
-#include "Framework/MultiBox/MultiBoxBuilder.h"
-//#include "Widgets/Colors/SColorPicker.h"
-#include "Framework/Commands/GenericCommands.h"
-//#include "NativeClassHierarchy.h"
-#include "Framework/Notifications/NotificationManager.h"
-#include "Widgets/Notifications/SNotificationList.h"
-//#include "ContentBrowserCommands.h"
 #include "ToolMenus.h"
 #include "ContentBrowserMenuContexts.h"
-#include "IContentBrowserDataModule.h"
-#include "ContentBrowserDataSource.h"
-#include "ContentBrowserDataSubsystem.h"
-#include "ContentBrowserDataMenuContexts.h"
-#include "Stats/Stats.h"
-#include "AssetContextMenuUtils.h"
 #include "Misc/ScopedSlowTask.h"
 
 #define LOCTEXT_NAMESPACE "ContentBrowser"
@@ -72,18 +37,6 @@ void FAssetFolderContextMenu::AddMenuOptions(UToolMenu* Menu)
 	UContentBrowserDataMenuContext_FolderMenu* Context = Menu->FindContext<UContentBrowserDataMenuContext_FolderMenu>();
 	checkf(Context, TEXT("Required context UContentBrowserDataMenuContext_FolderMenu was missing!"));
 
-	// Cache any vars that are used in determining if you can execute any actions.
-	// Useful for actions whose "CanExecute" will not change or is expensive to calculate.
-	StartProcessCanExecuteVars();
-	TWeakPtr<FAssetFolderContextMenu> WeakThisPtr = AsShared();
-	Menu->Context.AddCleanup([WeakThisPtr]()
-		{
-			if (TSharedPtr<FAssetFolderContextMenu> SharedThis = WeakThisPtr.Pin())
-			{
-				SharedThis->StopProcessCanExecuteVars();
-			}
-		});
-
 	if(Context->bCanBeModified)
 	{
 		// Bulk operations section //
@@ -115,81 +68,6 @@ void FAssetFolderContextMenu::AddMenuOptions(UToolMenu* Menu)
 				Entry.InsertPosition = FToolMenuInsert("FixUpRedirectorsInFolder", EToolMenuInsertType::After);
 			}
 		}
-
-		// Source control section //
-		{
-			FToolMenuSection& Section = Menu->AddSection("PathContextSourceControl", LOCTEXT("AssetTreeSCCMenuHeading", "Source Control"));
-
-			ISourceControlProvider& SourceControlProvider = ISourceControlModule::Get().GetProvider();
-			if (SourceControlProvider.IsEnabled())
-			{
-				using namespace UE::ContentBrowserAssetDataSource::Private;
-
-				// Check out
-				AddAsyncMenuEntry(
-					Section,
-					"FolderSCCCheckOut",
-					LOCTEXT("FolderSCCCheckOut", "Check Out"),
-					LOCTEXT("FolderSCCCheckOutTooltip", "Checks out all assets from source control which are in this selection."),
-					FUIAction(
-						FExecuteAction::CreateSP(this, &FAssetFolderContextMenu::ExecuteSCCCheckOut),
-						FCanExecuteAction::CreateSP(this, &FAssetFolderContextMenu::CanExecuteSCCCheckOut)
-					),
-					FIsAsyncProcessingActive::CreateSP(this, &FAssetFolderContextMenu::IsProcessingSCCCheckOut)
-				);
-
-				// Open for Add
-				AddAsyncMenuEntry(
-					Section,
-					"FolderSCCOpenForAdd",
-					LOCTEXT("FolderSCCOpenForAdd", "Mark For Add"),
-					LOCTEXT("FolderSCCOpenForAddTooltip", "Adds all assets to source control that are in this selection and not already added."),
-					FUIAction(
-						FExecuteAction::CreateSP(this, &FAssetFolderContextMenu::ExecuteSCCOpenForAdd),
-						FCanExecuteAction::CreateSP(this, &FAssetFolderContextMenu::CanExecuteSCCOpenForAdd)
-					),
-					FIsAsyncProcessingActive::CreateSP(this, &FAssetFolderContextMenu::IsProcessingSCCOpenForAdd)
-				);
-
-				// Check in
-				AddAsyncMenuEntry(
-					Section,
-					"FolderSCCCheckIn",
-					LOCTEXT("FolderSCCCheckIn", "Check In"),
-					LOCTEXT("FolderSCCCheckInTooltip", "Checks in all assets to source control which are in this selection."),
-					FUIAction(
-						FExecuteAction::CreateSP(this, &FAssetFolderContextMenu::ExecuteSCCCheckIn),
-						FCanExecuteAction::CreateSP(this, &FAssetFolderContextMenu::CanExecuteSCCCheckIn)
-					),
-					FIsAsyncProcessingActive::CreateSP(this, &FAssetFolderContextMenu::IsProcessingSCCCheckIn)
-				);
-
-				// Sync
-				Section.AddMenuEntry(
-					"FolderSCCSync",
-					LOCTEXT("FolderSCCSync", "Sync"),
-					LOCTEXT("FolderSCCSyncTooltip", "Syncs all the assets in this selection to the latest version."),
-					FSlateIcon(),
-					FUIAction(
-						FExecuteAction::CreateSP(this, &FAssetFolderContextMenu::ExecuteSCCSync),
-						FCanExecuteAction::CreateSP(this, &FAssetFolderContextMenu::CanExecuteSCCSync)
-					)
-				);
-			}
-			else
-			{
-				Section.AddMenuEntry(
-					"FolderSCCConnect",
-					LOCTEXT("FolderSCCConnect", "Connect To Source Control"),
-					LOCTEXT("FolderSCCConnectTooltip", "Connect to source control to allow source control operations to be performed on content and levels."),
-					FSlateIcon(FAppStyle::GetAppStyleSetName(), "MainFrame.ConnectToSourceControl"),
-					FUIAction(
-						FExecuteAction::CreateSP(this, &FAssetFolderContextMenu::ExecuteSCCConnect),
-						FCanExecuteAction::CreateSP(this, &FAssetFolderContextMenu::CanExecuteSCCConnect)
-					)
-				);
-			}
-		}
 	}
 }
 
@@ -215,8 +93,8 @@ void FAssetFolderContextMenu::ExecuteMigrateFolder()
 			// Don't collect the external packages here
 			if (AssetIt->GetOptionalOuterPathName().IsNone())
 			{
-				PackageNames.Add((*AssetIt).PackageName);
-			}
+			PackageNames.Add((*AssetIt).PackageName);
+		}
 		}
 
 		PackageNames.Append(SelectedPackages);
@@ -299,314 +177,9 @@ void FAssetFolderContextMenu::ExecuteFixUpRedirectorsInFolder()
 	}
 }
 
-void FAssetFolderContextMenu::ExecuteSCCCheckOut()
-{
-	// Get a list of package names in the selected paths
-	TArray<FString> PackageNames;
-	GetPackageNamesInSelectedPaths(PackageNames);
-	PackageNames.Append(SelectedPackages);
-
-	TArray<UPackage*> PackagesToCheckOut;
-	for ( auto PackageIt = PackageNames.CreateConstIterator(); PackageIt; ++PackageIt )
-	{
-		if ( FPackageName::DoesPackageExist(*PackageIt) )
-		{
-			// Since the file exists, create the package if it isn't loaded or just find the one that is already loaded
-			// No need to load unloaded packages. It isn't needed for the checkout process
-			UPackage* Package = CreatePackage(**PackageIt);
-			PackagesToCheckOut.Add( CreatePackage(**PackageIt) );
-		}
-	}
-
-	if ( PackagesToCheckOut.Num() > 0 )
-	{
-		// Update the source control status of all potentially relevant packages
-		ISourceControlModule::Get().GetProvider().Execute(ISourceControlOperation::Create<FUpdateStatus>(), PackagesToCheckOut);
-
-		// Now check them out
-		FEditorFileUtils::CheckoutPackages(PackagesToCheckOut);
-	}
-}
-
-void FAssetFolderContextMenu::ExecuteSCCOpenForAdd()
-{
-	ISourceControlProvider& SourceControlProvider = ISourceControlModule::Get().GetProvider();
-
-	// Get a list of package names in the selected paths
-	TArray<FString> PackageNames;
-	GetPackageNamesInSelectedPaths(PackageNames);
-	PackageNames.Append(SelectedPackages);
-
-	TArray<FString> PackagesToAdd;
-	TArray<UPackage*> PackagesToSave;
-	for ( auto PackageIt = PackageNames.CreateConstIterator(); PackageIt; ++PackageIt )
-	{
-		FSourceControlStatePtr SourceControlState = SourceControlProvider.GetState(SourceControlHelpers::PackageFilename(*PackageIt), EStateCacheUsage::Use);
-		if ( SourceControlState.IsValid() && !SourceControlState->IsSourceControlled() )
-		{
-			PackagesToAdd.Add(*PackageIt);
-
-			// Make sure the file actually exists on disk before adding it
-			FString Filename;
-			if ( !FPackageName::DoesPackageExist(*PackageIt, &Filename) )
-			{
-				UPackage* Package = FindPackage(NULL, **PackageIt);
-				if ( Package )
-				{
-					PackagesToSave.Add(Package);
-				}
-			}
-		}
-	}
-
-	if ( PackagesToAdd.Num() > 0 )
-	{
-		// If any of the packages are new, save them now
-		if ( PackagesToSave.Num() > 0 )
-		{
-			const bool bCheckDirty = false;
-			const bool bPromptToSave = false;
-			TArray<UPackage*> FailedPackages;
-			const FEditorFileUtils::EPromptReturnCode Return = FEditorFileUtils::PromptForCheckoutAndSave(PackagesToSave, bCheckDirty, bPromptToSave, &FailedPackages);
-			if(FailedPackages.Num() > 0)
-			{
-				// don't try and add files that failed to save - remove them from the list
-				for(auto FailedPackageIt = FailedPackages.CreateConstIterator(); FailedPackageIt; FailedPackageIt++)
-				{
-					PackagesToAdd.Remove((*FailedPackageIt)->GetName());
-				}
-			}
-		}
-
-		if ( PackagesToAdd.Num() > 0 )
-		{
-			SourceControlProvider.Execute(ISourceControlOperation::Create<FMarkForAdd>(), SourceControlHelpers::PackageFilenames(PackagesToAdd));
-		}
-	}
-}
-
-void FAssetFolderContextMenu::ExecuteSCCCheckIn()
-{
-	// Get a list of package names in the selected paths
-	TArray<FString> PackageNames;
-	GetPackageNamesInSelectedPaths(PackageNames);
-	PackageNames.Append(SelectedPackages);
-
-	// Form a list of loaded packages to prompt for save
-	TArray<UPackage*> LoadedPackages;
-	for ( auto PackageIt = PackageNames.CreateConstIterator(); PackageIt; ++PackageIt )
-	{
-		UPackage* Package = FindPackage(NULL, **PackageIt);
-		if ( Package )
-		{
-			LoadedPackages.Add(Package);
-		}
-	}
-
-	// Prompt the user to ask if they would like to first save any dirty packages they are trying to check-in
-	const FEditorFileUtils::EPromptReturnCode UserResponse = FEditorFileUtils::PromptForCheckoutAndSave( LoadedPackages, true, true );
-
-	// If the user elected to save dirty packages, but one or more of the packages failed to save properly OR if the user
-	// canceled out of the prompt, don't follow through on the check-in process
-	const bool bShouldProceed = ( UserResponse == FEditorFileUtils::EPromptReturnCode::PR_Success || UserResponse == FEditorFileUtils::EPromptReturnCode::PR_Declined );
-	if ( bShouldProceed )
-	{
-		TArray<FString> PendingDeletePaths;
-		for (const auto& Path : SelectedPaths)
-		{
-			PendingDeletePaths.Add(FPaths::ConvertRelativePathToFull(FPackageName::LongPackageNameToFilename(Path + TEXT("/"))));
-		}
-
-		const bool bUseSourceControlStateCache = false;
-		FSourceControlWindows::PromptForCheckin(bUseSourceControlStateCache, PackageNames, PendingDeletePaths);
-	}
-	else
-	{
-		// If a failure occurred, alert the user that the check-in was aborted. This warning shouldn't be necessary if the user cancelled
-		// from the dialog, because they obviously intended to cancel the whole operation.
-		if ( UserResponse == FEditorFileUtils::EPromptReturnCode::PR_Failure )
-		{
-			FMessageDialog::Open( EAppMsgType::Ok, NSLOCTEXT("UnrealEd", "SCC_Checkin_Aborted", "Check-in aborted as a result of save failure.") );
-		}
-	}
-}
-
-void FAssetFolderContextMenu::ExecuteSCCSync() const
-{
-	AssetViewUtils::SyncPathsFromSourceControl(SelectedPaths);
-	AssetViewUtils::SyncPackagesFromSourceControl(SelectedPackages);
-}
-
-void FAssetFolderContextMenu::ExecuteSCCConnect() const
-{
-	ISourceControlModule::Get().ShowLoginDialog(FSourceControlLoginClosed(), ELoginWindowMode::Modeless);
-}
-
-bool FAssetFolderContextMenu::CanExecuteSCCCheckOut() const
-{
-	return bCanExecuteSCCCheckOut && (!SelectedPaths.IsEmpty() || !SelectedPackages.IsEmpty());
-}
-
-bool FAssetFolderContextMenu::CanExecuteSCCOpenForAdd() const
-{
-	return bCanExecuteSCCOpenForAdd && (!SelectedPaths.IsEmpty() || !SelectedPackages.IsEmpty());
-}
-
-bool FAssetFolderContextMenu::CanExecuteSCCCheckIn() const
-{
-	return bCanExecuteSCCCheckIn && (!SelectedPaths.IsEmpty() || !SelectedPackages.IsEmpty()) && bUsesFileRevisions;
-}
-
-bool FAssetFolderContextMenu::CanExecuteSCCSync() const
-{
-	if (!SelectedPaths.IsEmpty() || !SelectedPackages.IsEmpty())
-	{
-		ISourceControlProvider& SourceControlProvider = ISourceControlModule::Get().GetProvider();
-		if (SourceControlProvider.IsEnabled() && SourceControlProvider.UsesFileRevisions())
-		{
-			return true;
-		}
-	}
-	return false;
-}
-
-bool FAssetFolderContextMenu::CanExecuteSCCConnect() const
-{
-	return (!ISourceControlModule::Get().IsEnabled() || !ISourceControlModule::Get().GetProvider().IsAvailable()) && SelectedPaths.Num() > 0;
-}
-
-bool FAssetFolderContextMenu::IsProcessingSCCCheckOut() const
-{
-	return IsTickable() && !CanExecuteSCCCheckOut();
-}
-
-bool FAssetFolderContextMenu::IsProcessingSCCOpenForAdd() const
-{
-	return IsTickable() && !CanExecuteSCCOpenForAdd();
-}
-
-bool FAssetFolderContextMenu::IsProcessingSCCCheckIn() const
-{
-	return IsTickable() && !CanExecuteSCCCheckIn();
-}
-
-
-void FAssetFolderContextMenu::StartProcessCanExecuteVars()
-{
-	// Start Cache whether we can execute any of the source control commands
-	StopProcessCanExecuteVars();
-
-	bCanExecuteSCCCheckOut = false;
-	bCanExecuteSCCOpenForAdd = false;
-	bCanExecuteSCCCheckIn = false;
-	bUsesFileRevisions = false;
-
-	ISourceControlProvider& SourceControlProvider = ISourceControlModule::Get().GetProvider();
-	if ( SourceControlProvider.IsEnabled() && SourceControlProvider.IsAvailable() )
-	{
-		GetPackageNamesInSelectedPaths(PackageNamesToProcess);
-		PackageNamesToProcess.Append(SelectedPackages);
-		CurrentPackageIndex = 0;
-		ProcessCanExecuteVars();
-	}
-}
-
-void FAssetFolderContextMenu::GetPackageNamesInSelectedPaths(TArray<FString>& OutPackageNames) const
-{
-	// Load the asset registry module
-	FAssetRegistryModule& AssetRegistryModule = FModuleManager::Get().LoadModuleChecked<FAssetRegistryModule>(TEXT("AssetRegistry"));
-
-	// Form a filter from the paths
-	FARFilter Filter;
-	Filter.bRecursivePaths = true;
-	for (int32 PathIdx = 0; PathIdx < SelectedPaths.Num(); ++PathIdx)
-	{
-		const FString& Path = SelectedPaths[PathIdx];
-		Filter.PackagePaths.Add(FName(*Path));
-	}
-
-	// Query for a list of assets in the selected paths
-	TArray<FAssetData> AssetList;
-	AssetRegistryModule.Get().GetAssets(Filter, AssetList);
-
-	// Form a list of unique package names from the assets
-	TSet<FName> UniquePackageNames;
-	for (int32 AssetIdx = 0; AssetIdx < AssetList.Num(); ++AssetIdx)
-	{
-		UniquePackageNames.Add(AssetList[AssetIdx].PackageName);
-	}
-
-	// Add all unique package names to the output
-	for ( auto PackageIt = UniquePackageNames.CreateConstIterator(); PackageIt; ++PackageIt )
-	{
-		OutPackageNames.Add( (*PackageIt).ToString() );
-	}
-}
-
 FString FAssetFolderContextMenu::GetFirstSelectedPath() const
 {
 	return SelectedPaths.Num() > 0 ? SelectedPaths[0] : TEXT("");
 }
 
-void FAssetFolderContextMenu::Tick(float DeltaTime)
-{
-	ProcessCanExecuteVars();
-}
-
-TStatId FAssetFolderContextMenu::GetStatId() const
-{
-	RETURN_QUICK_DECLARE_CYCLE_STAT(FAssetFolderContextMenu, STATGROUP_Tickables);
-}
-
-void FAssetFolderContextMenu::ProcessCanExecuteVars()
-{
-	const int32 StartPackageIndex = CurrentPackageIndex;
-	constexpr int32 MaxToProcessThisFrame = 500;
-
-	ISourceControlProvider& SourceControlProvider = ISourceControlModule::Get().GetProvider();
-
-	while (PackageNamesToProcess.Num() > CurrentPackageIndex
-		&& (CurrentPackageIndex - StartPackageIndex) <= MaxToProcessThisFrame)
-	{
-		FSourceControlStatePtr SourceControlState = SourceControlProvider.GetState(SourceControlHelpers::PackageFilename(*PackageNamesToProcess[CurrentPackageIndex]), EStateCacheUsage::Use);
-		++CurrentPackageIndex;
-		if (SourceControlState.IsValid())
-		{
-			if (!bCanExecuteSCCCheckOut && SourceControlState->CanCheckout())
-			{
-				bCanExecuteSCCCheckOut = true;
-			}
-			else if (!bCanExecuteSCCOpenForAdd && !SourceControlState->IsSourceControlled())
-			{
-				bCanExecuteSCCOpenForAdd = true;
-			}
-			else if (!bCanExecuteSCCCheckIn && SourceControlState->CanCheckIn())
-			{
-				bCanExecuteSCCCheckIn = true;
-			}
-		}
-
-		if (bCanExecuteSCCCheckOut && bCanExecuteSCCOpenForAdd && bCanExecuteSCCCheckIn)
-		{
-			// All SCC options are available, no need to keep iterating
-			StopProcessCanExecuteVars();
-			break;
-		}
-	}
-
-	if (PackageNamesToProcess.Num() <= CurrentPackageIndex)
-	{
-		// Everything was processed stop ticking
-		StopProcessCanExecuteVars();
-	}
-
-	bUsesFileRevisions = SourceControlProvider.UsesFileRevisions();
-}
-
-void FAssetFolderContextMenu::StopProcessCanExecuteVars()
-{
-	PackageNamesToProcess.Empty();
-}
-
-
 #undef LOCTEXT_NAMESPACE
diff --git a/Engine/Plugins/Editor/ContentBrowser/ContentBrowserAssetDataSource/Source/ContentBrowserAssetDataSource/Private/ContentBrowserAssetDataCore.cpp b/Engine/Plugins/Editor/ContentBrowser/ContentBrowserAssetDataSource/Source/ContentBrowserAssetDataSource/Private/ContentBrowserAssetDataCore.cpp
--- a/Engine/Plugins/Editor/ContentBrowser/ContentBrowserAssetDataSource/Source/ContentBrowserAssetDataSource/Private/ContentBrowserAssetDataCore.cpp
+++ b/Engine/Plugins/Editor/ContentBrowser/ContentBrowserAssetDataSource/Source/ContentBrowserAssetDataSource/Private/ContentBrowserAssetDataCore.cpp
@@ -18,6 +18,7 @@
 #include "AssetFolderContextMenu.h"
 #include "AssetFileContextMenu.h"
 #include "ContentBrowserDataUtils.h"
+#include "ContentBrowserDataMenuContexts.h"
 #include "Settings/ContentBrowserSettings.h"
 #include "Engine/Level.h"
 #include "Framework/Notifications/NotificationManager.h"
diff --git a/Engine/Plugins/Editor/ContentBrowser/ContentBrowserAssetDataSource/Source/ContentBrowserAssetDataSource/Public/AssetFileContextMenu.h b/Engine/Plugins/Editor/ContentBrowser/ContentBrowserAssetDataSource/Source/ContentBrowserAssetDataSource/Public/AssetFileContextMenu.h
--- a/Engine/Plugins/Editor/ContentBrowser/ContentBrowserAssetDataSource/Source/ContentBrowserAssetDataSource/Public/AssetFileContextMenu.h
+++ b/Engine/Plugins/Editor/ContentBrowser/ContentBrowserAssetDataSource/Source/ContentBrowserAssetDataSource/Public/AssetFileContextMenu.h
@@ -4,8 +4,6 @@
 
 #include "CoreMinimal.h"
 #include "AssetRegistry/AssetData.h"
-#include "ContentBrowserDataMenuContexts.h"
-#include "Input/Reply.h"
 
 class UToolMenu;
 
@@ -110,15 +108,6 @@ private:
 	/** Adds asset documentation menu options to a menu builder. Returns true if any options were added. */
 	bool AddDocumentationMenuOptions(UToolMenu* Menu);
 	
-	/** Adds source control menu options to a menu builder. */
-	bool AddSourceControlMenuOptions(UToolMenu* Menu);
-
-	/** Fills the source control sub-menu */
-	void FillSourceControlSubMenu(UToolMenu* Menu);
-
-	/** Handler to check to see if SCC actions are allowed */
-	bool CanExecuteSourceControlActions() const;
-
 	/** Creates a sub-menu of Chunk IDs that are are assigned to all selected assets */
 	void MakeChunkIDListMenu(UToolMenu* Menu);
 
@@ -173,36 +162,6 @@ private:
 	/** Handler for Bulk Export */
 	void ExecuteBulkExport();
 
-	/** Handler for when "Refresh source control" is selected */
-	void ExecuteSCCRefresh();
-
-	/** Handler for when "Merge" is selected */
-	void ExecuteSCCMerge();
-
-	/** Handler for when "Checkout from source control" is selected */
-	void ExecuteSCCCheckOut();
-
-	/** Handler for when "Open for add to source control" is selected */
-	void ExecuteSCCOpenForAdd();
-
-	/** Handler for when "Checkin to source control" is selected */
-	void ExecuteSCCCheckIn();
-
-	/** Handler for when "Source Control History" is selected */
-	void ExecuteSCCHistory();
-
-	/** Handler for when "Diff Against Depot" is selected */
-	void ExecuteSCCDiffAgainstDepot() const;
-
-	/** Handler for when "Source Control Revert" is selected */
-	void ExecuteSCCRevert();
-
-	/** Handler for when "Source Control Sync" is selected */
-	void ExecuteSCCSync();
-
-	/** Handler for when source control is disabled */
-	void ExecuteEnableSourceControl();
-
 	/** Handler to assign ChunkID to a selection of assets */
 	void ExecuteAssignChunkID();
 
@@ -238,36 +197,6 @@ private:
 	/** Handler to check to see if a "Show MetaData" command is allowed */
 	bool CanExecuteShowAssetMetaData() const;
 
-	/** Handler to check to see if "Refresh source control" can be executed */
-	bool CanExecuteSCCRefresh() const;
-
-	/** Handler to check to see if "Merge" can be executed */
-	bool CanExecuteSCCMerge() const;
-
-	/** Handler to check to see if "Checkout from source control" can be executed */
-	bool CanExecuteSCCCheckOut() const;
-
-	/** Handler to check to see if "Open for add to source control" can be executed */
-	bool CanExecuteSCCOpenForAdd() const;
-
-	/** Handler to check to see if "Checkin to source control" can be executed */
-	bool CanExecuteSCCCheckIn() const;
-
-	/** Handler to check to see if "Source Control History" can be executed */
-	bool CanExecuteSCCHistory() const;
-
-	/** Handler to check to see if "Source Control Revert" can be executed */
-	bool CanExecuteSCCRevert() const;
-
-	/** Handler to check to see if "Source Control Sync" can be executed */
-	bool CanExecuteSCCSync() const;
-
-	/** Handler to check to see if "Diff Against Depot" can be executed */
-	bool CanExecuteSCCDiffAgainstDepot() const;
-
-	/** Handler to check to see if "Diff Selected" can be executed */
-	bool CanExecuteDiffSelected() const;
-
 	/** Handler to check to see if "Capture Thumbnail" can be executed */
 	bool CanExecuteCaptureThumbnail() const;
 
@@ -308,18 +237,9 @@ private:
 
 	FOnShowAssetsInPathsView OnShowAssetsInPathsView;
 
-	UContentBrowserDataMenuContext_FileMenu::FOnRefreshView OnRefreshView;
-
 	/** Cached CanExecute vars */
 	bool bAtLeastOneNonRedirectorSelected = false;
-	bool bCanExecuteSCCMerge = false;
-	bool bCanExecuteSCCCheckOut = false;
-	bool bCanExecuteSCCOpenForAdd = false;
-	bool bCanExecuteSCCCheckIn = false;
-	bool bCanExecuteSCCHistory = false;
-	bool bCanExecuteSCCRevert = false;
-	bool bCanExecuteSCCSync = false;
-	bool bUsesFileRevisions = false;
+
 	/** */
 	int32 ChunkIDSelected = 0;
 };
diff --git a/Engine/Plugins/Editor/ContentBrowser/ContentBrowserAssetDataSource/Source/ContentBrowserAssetDataSource/Public/AssetFolderContextMenu.h b/Engine/Plugins/Editor/ContentBrowser/ContentBrowserAssetDataSource/Source/ContentBrowserAssetDataSource/Public/AssetFolderContextMenu.h
--- a/Engine/Plugins/Editor/ContentBrowser/ContentBrowserAssetDataSource/Source/ContentBrowserAssetDataSource/Public/AssetFolderContextMenu.h
+++ b/Engine/Plugins/Editor/ContentBrowser/ContentBrowserAssetDataSource/Source/ContentBrowserAssetDataSource/Public/AssetFolderContextMenu.h
@@ -3,20 +3,12 @@
 #pragma once
 
 #include "CoreMinimal.h"
-#include "Input/Reply.h"
-#include "TickableEditorObject.h"
 
-class FExtender;
-class FMenuBuilder;
-class SWidget;
-class SWindow;
 class UToolMenu;
 
-class CONTENTBROWSERASSETDATASOURCE_API FAssetFolderContextMenu : public TSharedFromThis<FAssetFolderContextMenu>, public FTickableEditorObject
+class CONTENTBROWSERASSETDATASOURCE_API FAssetFolderContextMenu : public TSharedFromThis<FAssetFolderContextMenu>
 {
 public:
-	virtual ~FAssetFolderContextMenu() = default;
-
 	UE_DEPRECATED(5.1, "Use the new overload.")
 	void MakeContextMenu(
 		UToolMenu* InMenu,
@@ -45,77 +37,12 @@ private:
 	/** Handler for when "Fix up Redirectors in Folder" is selected */
 	void ExecuteFixUpRedirectorsInFolder();
 
-	/** Handler for when "Checkout from source control" is selected */
-	void ExecuteSCCCheckOut();
-
-	/** Handler for when "Open for Add to source control" is selected */
-	void ExecuteSCCOpenForAdd();
-
-	/** Handler for when "Checkin to source control" is selected */
-	void ExecuteSCCCheckIn();
-
-	/** Handler for when "Sync from source control" is selected */
-	void ExecuteSCCSync() const;
-
-	/** Handler for when "Connect to source control" is selected */
-	void ExecuteSCCConnect() const;
-
-	/** Handler to check to see if "Checkout from source control" can be executed */
-	bool CanExecuteSCCCheckOut() const;
-
-	/** Handler to check to see if "Open for Add to source control" can be executed */
-	bool CanExecuteSCCOpenForAdd() const;
-
-	/** Handler to check to see if "Checkin to source control" can be executed */
-	bool CanExecuteSCCCheckIn() const;
-
-	/** Handler to check to see if "Sync" can be executed */
-	bool CanExecuteSCCSync() const;
-
-	/** Handler to check to see if "Connect to source control" can be executed */
-	bool CanExecuteSCCConnect() const;	
-
-	bool IsProcessingSCCCheckOut() const;
-
-	bool IsProcessingSCCOpenForAdd() const;
-
-	bool IsProcessingSCCCheckIn() const;
-
-	/** Initializes and start to process some variable used to in "CanExecute" checks that won't change at runtime or are too expensive to check every frame. */
-	void StartProcessCanExecuteVars();
-
-	/** Returns a list of names of packages in all selected paths in the sources view */
-	void GetPackageNamesInSelectedPaths(TArray<FString>& OutPackageNames) const;
-
 	/** Gets the first selected path, if it exists */
 	FString GetFirstSelectedPath() const;
 
-	/** Begin FTickableEditorObject interface */
-	virtual void Tick(float DeltaTime) override;
-	virtual bool IsTickable() const override { return !PackageNamesToProcess.IsEmpty(); }
-	virtual TStatId GetStatId() const override;
-	/** End FTickableEditorObject interface */
 private:
-
-	void ProcessCanExecuteVars();
-
-	/** Stop or cancel the process of some variable used to in "CanExecute" checks */
-	void StopProcessCanExecuteVars();
-
 	TArray<FString> SelectedPaths;
 
 	// Not null if the selection contained some assets also
 	TArray<FString> SelectedPackages;
-
-	TWeakPtr<SWidget> ParentWidget;
-
-	/** Cached SCC CanExecute vars */
-	bool bCanExecuteSCCCheckOut = false;
-	bool bCanExecuteSCCOpenForAdd = false;
-	bool bCanExecuteSCCCheckIn = false;
-	bool bUsesFileRevisions = false;
-
-
-	TArray<FString> PackageNamesToProcess;
-	int32 CurrentPackageIndex = 0;
 };
diff --git a/Engine/Source/Developer/AssetTools/Private/AssetViewUtils.cpp b/Engine/Source/Developer/AssetTools/Private/AssetViewUtils.cpp
--- a/Engine/Source/Developer/AssetTools/Private/AssetViewUtils.cpp
+++ b/Engine/Source/Developer/AssetTools/Private/AssetViewUtils.cpp
@@ -1742,7 +1742,7 @@ void AssetViewUtils::SyncPackagesFromSourceControl(const TArray<FString>& Packag
 		UPackageTools::UnloadPackages(PackagesToUnload);
 
 		// Re-cache the SCC state...
-		SCCProvider.Execute(ISourceControlOperation::Create<FUpdateStatus>(), PackageFilenames, EConcurrency::Asynchronous);
+		SCCProvider.Execute(ISourceControlOperation::Create<FUpdateStatus>(), PackageFilenames);
 	}
 }
 
@@ -1874,7 +1874,7 @@ void AssetViewUtils::SyncPathsFromSourceControl(const TArray<FString>& ContentPa
 		UPackageTools::UnloadPackages(PackageRawPtrsToUnload);
 
 		// Re-cache the SCC state...
-		SCCProvider.Execute(ISourceControlOperation::Create<FUpdateStatus>(), PathsOnDisk, EConcurrency::Asynchronous);
+		SCCProvider.Execute(ISourceControlOperation::Create<FUpdateStatus>(), PathsOnDisk);
 	}
 }
 
