Fix compile error, missing headers and forward declares

diff --git a/Engine/Plugins/Editor/AssetManagerEditor/Source/AssetManagerEditor/Private/AssetSourceControlContextMenu.cpp b/Engine/Plugins/Editor/AssetManagerEditor/Source/AssetManagerEditor/Private/AssetSourceControlContextMenu.cpp
--- a/Engine/Plugins/Editor/AssetManagerEditor/Source/AssetManagerEditor/Private/AssetSourceControlContextMenu.cpp
+++ b/Engine/Plugins/Editor/AssetManagerEditor/Source/AssetManagerEditor/Private/AssetSourceControlContextMenu.cpp
@@ -9,6 +9,7 @@
 #include "Async/Async.h"
 #include "ContentBrowserMenuContexts.h"
 #include "FileHelpers.h"
+#include "HAL/PlatformFileManager.h"
 #include "ISourceControlModule.h"
 #include "Misc/FileHelper.h"
 #include "Misc/MessageDialog.h"
@@ -28,6 +29,7 @@
 #include "ToolMenuDelegates.h"
 #include "ToolMenuEntry.h"
 #include "ToolMenuSection.h"
+#include "Widgets/Images/SImage.h"
 #include "Widgets/Images/SThrobber.h"
 #include "Widgets/SBoxPanel.h"
 #include "Widgets/Text/STextBlock.h"
diff --git a/Engine/Plugins/Editor/ContentBrowser/ContentBrowserAssetDataSource/Source/ContentBrowserAssetDataSource/Public/AssetFileContextMenu.h b/Engine/Plugins/Editor/ContentBrowser/ContentBrowserAssetDataSource/Source/ContentBrowserAssetDataSource/Public/AssetFileContextMenu.h
--- a/Engine/Plugins/Editor/ContentBrowser/ContentBrowserAssetDataSource/Source/ContentBrowserAssetDataSource/Public/AssetFileContextMenu.h
+++ b/Engine/Plugins/Editor/ContentBrowser/ContentBrowserAssetDataSource/Source/ContentBrowserAssetDataSource/Public/AssetFileContextMenu.h
@@ -4,8 +4,11 @@
 
 #include "CoreMinimal.h"
 #include "AssetRegistry/AssetData.h"
+#include "Input/Reply.h"
 
 class UToolMenu;
+class SWindow;
+class SWidget;
 
 class CONTENTBROWSERASSETDATASOURCE_API FAssetFileContextMenu : public TSharedFromThis<FAssetFileContextMenu>
 {
