Fix duplicated localization keys

diff --git a/Engine/Plugins/Editor/AssetManagerEditor/Source/AssetManagerEditor/Private/AssetSourceControlContextMenu.cpp b/Engine/Plugins/Editor/AssetManagerEditor/Source/AssetManagerEditor/Private/AssetSourceControlContextMenu.cpp
--- a/Engine/Plugins/Editor/AssetManagerEditor/Source/AssetManagerEditor/Private/AssetSourceControlContextMenu.cpp
+++ b/Engine/Plugins/Editor/AssetManagerEditor/Source/AssetManagerEditor/Private/AssetSourceControlContextMenu.cpp
@@ -409,7 +409,7 @@ void FAssetSourceControlContextMenuState::FillSourceControlSubMenu(UToolMenu* Me
 						return LOCTEXT("SCCCheckOutTooltip", "Check out the selected assets from source control.");
 					}
 
-					return FText::Format(LOCTEXT("SCCPartialSyncAndCheckOut", "Checks out the selected assets from source control that are not currently locked.\n\nLocked Assets:\n{0}"), CheckedOutUsersText);
+					return FText::Format(LOCTEXT("SCCPartialCheckOut", "Checks out the selected assets from source control that are not currently locked.\n\nLocked Assets:\n{0}"), CheckedOutUsersText);
 				}),
 
 			TAttribute<FSlateIcon>::CreateLambda([this]()
@@ -531,7 +531,7 @@ void FAssetSourceControlContextMenuState::FillSourceControlSubMenu(UToolMenu* Me
 		AddAsyncMenuEntry(Section,
 			"SCCRevertWritable",
 			LOCTEXT("SCCRevertWritable", "Revert Writable Files"),
-			LOCTEXT("SCCRevertTooltip", "Reverts the assets that are Writable to their current state from source control. They will remain at their current revision."),
+			LOCTEXT("SCCRevertWritableTooltip", "Reverts the assets that are Writable to their current state from source control. They will remain at their current revision."),
 			FSlateIcon(FAppStyle::GetAppStyleSetName(), "SourceControl.Actions.Revert"),
 			FUIAction(
 				ExecutionCheck(FExecuteAction::CreateSP(this, &FAssetSourceControlContextMenuState::ExecuteSCCRevertWritable)),
