Fix some source control windows not appearing: Don't wrap operations that have their own windows in slow scope since it will destroy all windows that got attached to it when it goes out of scope.

diff --git a/Engine/Plugins/Editor/AssetManagerEditor/Source/AssetManagerEditor/Private/AssetSourceControlContextMenu.cpp b/Engine/Plugins/Editor/AssetManagerEditor/Source/AssetManagerEditor/Private/AssetSourceControlContextMenu.cpp
--- a/Engine/Plugins/Editor/AssetManagerEditor/Source/AssetManagerEditor/Private/AssetSourceControlContextMenu.cpp
+++ b/Engine/Plugins/Editor/AssetManagerEditor/Source/AssetManagerEditor/Private/AssetSourceControlContextMenu.cpp
@@ -497,7 +497,7 @@ void FAssetSourceControlContextMenuState::FillSourceControlSubMenu(UToolMenu* Me
 		LOCTEXT("SCCHistoryTooltip", "Displays the source control revision history of the selected assets."),
 		FSlateIcon(FAppStyle::GetAppStyleSetName(), "SourceControl.Actions.History"),
 		FUIAction(
-			ExecutionCheck(FExecuteAction::CreateSP(this, &FAssetSourceControlContextMenuState::ExecuteSCCHistory)),
+			FExecuteAction::CreateSP(this, &FAssetSourceControlContextMenuState::ExecuteSCCHistory),
 			FCanExecuteAction::CreateLambda([this]() { return IsActionEnabled(CanExecuteSCCHistory()); })
 		),
 		FIsAsyncProcessingActive::CreateLambda([this]() { return IsStillScanning(CanExecuteSCCHistory()); })
@@ -509,7 +509,7 @@ void FAssetSourceControlContextMenuState::FillSourceControlSubMenu(UToolMenu* Me
 		LOCTEXT("SCCDiffAgainstDepotTooltip", "Look at differences between the local and remote version of the selected assets."),
 		FSlateIcon(FAppStyle::GetAppStyleSetName(), "SourceControl.Actions.Diff"),
 		FUIAction(
-			ExecutionCheck(FExecuteAction::CreateSP(this, &FAssetSourceControlContextMenuState::ExecuteSCCDiffAgainstDepot)),
+			FExecuteAction::CreateSP(this, &FAssetSourceControlContextMenuState::ExecuteSCCDiffAgainstDepot),
 			FCanExecuteAction::CreateLambda([this]() { return IsActionEnabled(CanExecuteSCCDiffAgainstDepot()); })
 		),
 		FIsAsyncProcessingActive::CreateLambda([this]() { return IsStillScanning(CanExecuteSCCDiffAgainstDepot()); })
@@ -521,7 +521,7 @@ void FAssetSourceControlContextMenuState::FillSourceControlSubMenu(UToolMenu* Me
 		LOCTEXT("SCCRevertTooltip", "Reverts the selected assets to their original state from source control."),
 		FSlateIcon(FAppStyle::GetAppStyleSetName(), "SourceControl.Actions.Revert"),
 		FUIAction(
-			ExecutionCheck(FExecuteAction::CreateSP(this, &FAssetSourceControlContextMenuState::ExecuteSCCRevert)),
+			FExecuteAction::CreateSP(this, &FAssetSourceControlContextMenuState::ExecuteSCCRevert),
 			FCanExecuteAction::CreateLambda([this]() { return IsActionEnabled(CanExecuteSCCRevert()); })
 		),
 		FIsAsyncProcessingActive::CreateLambda([this]() { return IsStillScanning(CanExecuteSCCRevert()); })
@@ -551,7 +551,7 @@ void FAssetSourceControlContextMenuState::FillSourceControlSubMenu(UToolMenu* Me
 			LOCTEXT("SCCMergeTooltip", "Opens the blueprint editor with the merge tool open."),
 			FSlateIcon(),
 			FUIAction(
-				ExecutionCheck(FExecuteAction::CreateSP(this, &FAssetSourceControlContextMenuState::ExecuteSCCMerge)),
+				FExecuteAction::CreateSP(this, &FAssetSourceControlContextMenuState::ExecuteSCCMerge),
 				FCanExecuteAction::CreateLambda([this]() { return IsActionEnabled(CanExecuteSCCMerge()); })
 			),
 			FIsAsyncProcessingActive::CreateLambda([this]() { return IsStillScanning(CanExecuteSCCMerge()); })
