Reference Viewer: Fix an issue where mouse interactions at the top of the graph were being ignored

diff --git a/Engine/Plugins/Editor/AssetManagerEditor/Source/AssetManagerEditor/Private/ReferenceViewer/SReferenceViewer.cpp b/Engine/Plugins/Editor/AssetManagerEditor/Source/AssetManagerEditor/Private/ReferenceViewer/SReferenceViewer.cpp
--- a/Engine/Plugins/Editor/AssetManagerEditor/Source/AssetManagerEditor/Private/ReferenceViewer/SReferenceViewer.cpp
+++ b/Engine/Plugins/Editor/AssetManagerEditor/Source/AssetManagerEditor/Private/ReferenceViewer/SReferenceViewer.cpp
@@ -227,7 +227,7 @@ void SReferenceViewer::Construct(const FArguments& InArgs)
 			.Padding(8)
 			[
 				SNew(SHorizontalBox)
-				.Visibility_Lambda([this]() { return ( Settings->GetFindPathEnabled() ? EVisibility::Collapsed : EVisibility::Visible); })
+				.Visibility_Lambda([this]() { return ( Settings->GetFindPathEnabled() ? EVisibility::Collapsed : EVisibility::SelfHitTestInvisible); })
 
 				+SHorizontalBox::Slot()
 				.AutoWidth()
