Replace some remaining instances of "Source Control" with "Revision Control" in text

diff --git a/Engine/Plugins/Editor/AssetManagerEditor/Source/AssetManagerEditor/Private/AssetSourceControlContextMenu.cpp b/Engine/Plugins/Editor/AssetManagerEditor/Source/AssetManagerEditor/Private/AssetSourceControlContextMenu.cpp
--- a/Engine/Plugins/Editor/AssetManagerEditor/Source/AssetManagerEditor/Private/AssetSourceControlContextMenu.cpp
+++ b/Engine/Plugins/Editor/AssetManagerEditor/Source/AssetManagerEditor/Private/AssetSourceControlContextMenu.cpp
@@ -318,8 +318,8 @@ bool FAssetSourceControlContextMenuState::AddSourceControlMenuOptions(FToolMenuS
 			// SCC sub menu
 			InSection.AddSubMenu(
 				"SourceControlSubMenu",
-				LOCTEXT("SourceControlSubMenuLabel", "Source Control"),
-				LOCTEXT("SourceControlSubMenuToolTip", "Source control actions."),
+				LOCTEXT("SourceControlSubMenuLabel", "Revision Control"),
+				LOCTEXT("SourceControlSubMenuToolTip", "Revision Control actions."),
 				FNewToolMenuDelegate::CreateSP(this, &FAssetSourceControlContextMenuState::FillSourceControlSubMenu),
 				false,
 				FSlateIcon(FAppStyle::GetAppStyleSetName(), "SourceControl.StatusIcon.On")
@@ -329,8 +329,8 @@ bool FAssetSourceControlContextMenuState::AddSourceControlMenuOptions(FToolMenuS
 		{
 			InSection.AddSubMenu(
 				"SourceControlSubMenu",
-				LOCTEXT("SourceControlSubMenuUnavailableLabel", "Source Control (Unavailable)"),
-				LOCTEXT("SourceControlSubMenuUnavailableToolTip", "Source control servers are currently unavailable, check your network connection."),
+				LOCTEXT("SourceControlSubMenuUnavailableLabel", "Revision Control (Unavailable)"),
+				LOCTEXT("SourceControlSubMenuUnavailableToolTip", "Revision Control servers are currently unavailable, check your network connection."),
 				FNewToolMenuDelegate(),
 				FUIAction(
 					FExecuteAction(),
@@ -346,8 +346,8 @@ bool FAssetSourceControlContextMenuState::AddSourceControlMenuOptions(FToolMenuS
 	{
 		InSection.AddMenuEntry(
 			"SCCConnectToSourceControl",
-			LOCTEXT("SCCConnectToSourceControl", "Connect To Source Control..."),
-			LOCTEXT("SCCConnectToSourceControlTooltip", "Connect to source control to allow source control operations to be performed on content and levels."),
+			LOCTEXT("SCCConnectToSourceControl", "Connect To Revision Control..."),
+			LOCTEXT("SCCConnectToSourceControlTooltip", "Connect to a revision control system for tracking changes to your content and levels."),
 			FSlateIcon(FAppStyle::GetAppStyleSetName(), "MainFrame.ConnectToSourceControl"),
 			FUIAction(
 				FExecuteAction::CreateLambda([]() { ISourceControlModule::Get().ShowLoginDialog(FSourceControlLoginClosed(), ELoginWindowMode::Modeless); }),
@@ -377,7 +377,7 @@ void FAssetSourceControlContextMenuState::FillSourceControlSubMenu(UToolMenu* Me
 {
 	ISourceControlProvider& SourceControlProvider = ISourceControlModule::Get().GetProvider();
 
-	FToolMenuSection& Section = Menu->AddSection("AssetSourceControlActions", LOCTEXT("AssetSourceControlActionsMenuHeading", "Source Control"));
+	FToolMenuSection& Section = Menu->AddSection("AssetSourceControlActions", LOCTEXT("AssetSourceControlActionsMenuHeading", "Revision Control"));
 
 	using namespace UE::AssetSourceControlContextMenu::Private;
 
@@ -390,7 +390,7 @@ void FAssetSourceControlContextMenuState::FillSourceControlSubMenu(UToolMenu* Me
 		AddAsyncMenuEntry(Section,
 			"SCCSync",
 			LOCTEXT("SCCSync", "Sync"),
-			LOCTEXT("SCCSyncTooltip", "Updates the selected assets to the latest version in source control."),
+			LOCTEXT("SCCSyncTooltip", "Updates the selected assets to the latest version in revision control."),
 			FSlateIcon(FAppStyle::GetAppStyleSetName(), "SourceControl.Actions.Sync"),
 			FUIAction(
 				ExecutionCheck(FExecuteAction::CreateSP(this, &FAssetSourceControlContextMenuState::ExecuteSCCSync)),
@@ -409,10 +409,10 @@ void FAssetSourceControlContextMenuState::FillSourceControlSubMenu(UToolMenu* Me
 				{
 					if (IsStillScanning(CanExecuteSCCCheckOut()) || CheckedOutUsers.Num() == 0)
 					{
-						return LOCTEXT("SCCCheckOutTooltip", "Check out the selected assets from source control.");
+						return LOCTEXT("SCCCheckOutTooltip", "Check out the selected assets from revision control.");
 					}
 
-					return FText::Format(LOCTEXT("SCCPartialCheckOut", "Checks out the selected assets from source control that are not currently locked.\n\nLocked Assets:\n{0}"), CheckedOutUsersText);
+					return FText::Format(LOCTEXT("SCCPartialCheckOut", "Checks out the selected assets from revision control that are not currently locked.\n\nLocked Assets:\n{0}"), CheckedOutUsersText);
 				}),
 
 			TAttribute<FSlateIcon>::CreateLambda([this]()
@@ -433,10 +433,10 @@ void FAssetSourceControlContextMenuState::FillSourceControlSubMenu(UToolMenu* Me
 				{
 					if (IsStillScanning(CanExecuteSCCCheckOut()) || CheckedOutUsers.Num() == 0)
 					{
-						return LOCTEXT("SCCSyncAndCheckOutTooltip", "Sync to latest and Check out the selected assets from source control.");
+						return LOCTEXT("SCCSyncAndCheckOutTooltip", "Sync to latest and Check out the selected assets from revision control.");
 					}
 
-					return FText::Format(LOCTEXT("SCCPartialSyncAndCheckOut", "Sync to latest and Checks out the selected assets from source control that are not currently locked.\n\nLocked Assets:\n{0}"), CheckedOutUsersText);
+					return FText::Format(LOCTEXT("SCCPartialSyncAndCheckOut", "Sync to latest and Checks out the selected assets from revision control that are not currently locked.\n\nLocked Assets:\n{0}"), CheckedOutUsersText);
 				}),
 
 			TAttribute<FSlateIcon>::CreateLambda([this]()
@@ -469,7 +469,7 @@ void FAssetSourceControlContextMenuState::FillSourceControlSubMenu(UToolMenu* Me
 	AddAsyncMenuEntry(Section,
 		"SCCOpenForAdd",
 		LOCTEXT("SCCOpenForAdd", "Mark For Add"),
-		LOCTEXT("SCCOpenForAddTooltip", "Adds the selected assets to source control."),
+		LOCTEXT("SCCOpenForAddTooltip", "Adds the selected assets to revision control."),
 		FSlateIcon(FAppStyle::GetAppStyleSetName(), "SourceControl.Actions.Add"),
 		FUIAction(
 			ExecutionCheck(FExecuteAction::CreateSP(this, &FAssetSourceControlContextMenuState::ExecuteSCCOpenForAdd)),
@@ -483,7 +483,7 @@ void FAssetSourceControlContextMenuState::FillSourceControlSubMenu(UToolMenu* Me
 		AddAsyncMenuEntry(Section,
 			"SCCCheckIn",
 			LOCTEXT("SCCCheckIn", "Check In"),
-			LOCTEXT("SCCCheckInTooltip", "Checks the selected assets into source control."),
+			LOCTEXT("SCCCheckInTooltip", "Checks the selected assets into revision control."),
 			FSlateIcon(FAppStyle::GetAppStyleSetName(), "SourceControl.Actions.Submit"),
 			FUIAction(
 				ExecutionCheck(FExecuteAction::CreateSP(this, &FAssetSourceControlContextMenuState::ExecuteSCCCheckIn)),
@@ -496,7 +496,7 @@ void FAssetSourceControlContextMenuState::FillSourceControlSubMenu(UToolMenu* Me
 	AddAsyncMenuEntry(Section,
 		"SCCHistory",
 		LOCTEXT("SCCHistory", "History"),
-		LOCTEXT("SCCHistoryTooltip", "Displays the source control revision history of the selected assets."),
+		LOCTEXT("SCCHistoryTooltip", "Displays the history of the selected asset in revision control."),
 		FSlateIcon(FAppStyle::GetAppStyleSetName(), "SourceControl.Actions.History"),
 		FUIAction(
 			FExecuteAction::CreateSP(this, &FAssetSourceControlContextMenuState::ExecuteSCCHistory),
@@ -520,7 +520,7 @@ void FAssetSourceControlContextMenuState::FillSourceControlSubMenu(UToolMenu* Me
 	AddAsyncMenuEntry(Section,
 		"SCCRevert",
 		LOCTEXT("SCCRevert", "Revert"),
-		LOCTEXT("SCCRevertTooltip", "Reverts the selected assets to their original state from source control."),
+		LOCTEXT("SCCRevertTooltip", "Reverts the selected assets to their original state from revision control."),
 		FSlateIcon(FAppStyle::GetAppStyleSetName(), "SourceControl.Actions.Revert"),
 		FUIAction(
 			FExecuteAction::CreateSP(this, &FAssetSourceControlContextMenuState::ExecuteSCCRevert),
@@ -534,7 +534,7 @@ void FAssetSourceControlContextMenuState::FillSourceControlSubMenu(UToolMenu* Me
 		AddAsyncMenuEntry(Section,
 			"SCCRevertWritable",
 			LOCTEXT("SCCRevertWritable", "Revert Writable Files"),
-			LOCTEXT("SCCRevertWritableTooltip", "Reverts the assets that are Writable to their current state from source control. They will remain at their current revision."),
+			LOCTEXT("SCCRevertWritableTooltip", "Reverts the assets that are Writable to their current state from revision control. They will remain at their current revision."),
 			FSlateIcon(FAppStyle::GetAppStyleSetName(), "SourceControl.Actions.Revert"),
 			FUIAction(
 				ExecutionCheck(FExecuteAction::CreateSP(this, &FAssetSourceControlContextMenuState::ExecuteSCCRevertWritable)),
@@ -563,7 +563,7 @@ void FAssetSourceControlContextMenuState::FillSourceControlSubMenu(UToolMenu* Me
 	Section.AddMenuEntry(
 		"SCCRefresh",
 		LOCTEXT("SCCRefresh", "Refresh"),
-		LOCTEXT("SCCRefreshTooltip", "Updates the source control status of the asset."),
+		LOCTEXT("SCCRefreshTooltip", "Updates the revision control status of the asset."),
 		FSlateIcon(FAppStyle::GetAppStyleSetName(), "SourceControl.Actions.Refresh"),
 		FUIAction(
 			ExecutionCheck(FExecuteAction::CreateSP(this, &FAssetSourceControlContextMenuState::ExecuteSCCRefresh)),
@@ -588,7 +588,7 @@ FExecuteAction FAssetSourceControlContextMenuState::ExecutionCheck(FExecuteActio
 				}
 			}
 
-			FScopedSlowTask SlowTask(0, LOCTEXT("SCCOperationSlowTaskLabel", "Performing Source Control Operation"));
+			FScopedSlowTask SlowTask(0, LOCTEXT("SCCOperationSlowTaskLabel", "Performing Revision Control Operation"));
 			SlowTask.MakeDialogDelayed(0.25f);
 			Action.Execute();
 		});
diff --git a/Engine/Plugins/Editor/ChangelistReview/Source/Private/SSourceControlReview.cpp b/Engine/Plugins/Editor/ChangelistReview/Source/Private/SSourceControlReview.cpp
--- a/Engine/Plugins/Editor/ChangelistReview/Source/Private/SSourceControlReview.cpp
+++ b/Engine/Plugins/Editor/ChangelistReview/Source/Private/SSourceControlReview.cpp
@@ -508,7 +508,7 @@ bool SSourceControlReview::IsChangelistRecordValid(const TArray<TMap<FString, FS
 	}
 	if (InRecord.Num() > 1)
 	{
-		FMessageDialog::Open(EAppMsgType::Ok, LOCTEXT("ChangelistInvalidResponseFormat", "Invalid API response from Source Control"));
+		FMessageDialog::Open(EAppMsgType::Ok, LOCTEXT("ChangelistInvalidResponseFormat", "Invalid API response from Revision Control"));
 		return false;
 	}
 
@@ -541,7 +541,7 @@ void SSourceControlReview::SetFileSourceControlAction(TSharedPtr<FChangelistFile
 {
 	if (!ChangelistFileData.IsValid())
 	{
-		UE_LOG(LogSourceControl, Error, TEXT("Unable to set source control action information. No changelist file data"));
+		UE_LOG(LogSourceControl, Error, TEXT("Unable to set revision control action information. No changelist file data"));
 		return;
 	}
 
@@ -567,7 +567,7 @@ void SSourceControlReview::SetFileSourceControlAction(TSharedPtr<FChangelistFile
 	}
 	else
 	{
-		UE_LOG(LogSourceControl, Error, TEXT("Unable to parse source control action information. '%s' diff will not be shown"), *ChangelistFileData->AssetName);
+		UE_LOG(LogSourceControl, Error, TEXT("Unable to parse revision control action information. '%s' diff will not be shown"), *ChangelistFileData->AssetName);
 		ChangelistFileData->FileSourceControlAction = ESourceControlAction::Unset;
 	}
 }
diff --git a/Engine/Source/Developer/UncontrolledChangelists/Private/UncontrolledChangelistState.cpp b/Engine/Source/Developer/UncontrolledChangelists/Private/UncontrolledChangelistState.cpp
--- a/Engine/Source/Developer/UncontrolledChangelists/Private/UncontrolledChangelistState.cpp
+++ b/Engine/Source/Developer/UncontrolledChangelists/Private/UncontrolledChangelistState.cpp
@@ -50,7 +50,7 @@ const FText& FUncontrolledChangelistState::GetDescriptionText() const
 
 FText FUncontrolledChangelistState::GetDisplayTooltip() const
 {
-	return LOCTEXT("UncontrolledStateTooltip", "Uncontrolled: Locally modified outside of source control");
+	return LOCTEXT("UncontrolledStateTooltip", "Uncontrolled: Locally modified outside of revision control");
 }
 
 const FDateTime& FUncontrolledChangelistState::GetTimeStamp() const
diff --git a/Engine/Source/Developer/Virtualization/Private/VirtualizationSourceControlBackend.cpp b/Engine/Source/Developer/Virtualization/Private/VirtualizationSourceControlBackend.cpp
--- a/Engine/Source/Developer/Virtualization/Private/VirtualizationSourceControlBackend.cpp
+++ b/Engine/Source/Developer/Virtualization/Private/VirtualizationSourceControlBackend.cpp
@@ -286,7 +286,7 @@ IVirtualizationBackend::EConnectionStatus FSourceControlBackend::OnConnect()
 	SCCProvider = ISourceControlModule::Get().CreateProvider(FName("Perforce"), TEXT("Virtualization"), SCCSettings);
 	if (!SCCProvider.IsValid())
 	{
-		UE_LOG(LogVirtualization, Error, TEXT("[%s] Failed to create a perforce source control provider"), *GetDebugName());
+		UE_LOG(LogVirtualization, Error, TEXT("[%s] Failed to create a perforce revision control provider"), *GetDebugName());
 		return IVirtualizationBackend::EConnectionStatus::Error;
 	}
 
@@ -305,7 +305,7 @@ IVirtualizationBackend::EConnectionStatus FSourceControlBackend::OnConnect()
 			}
 
 			FMessageLog Log("LogVirtualization");
-			Log.Warning(FText::Format(LOCTEXT("FailedSourceControlConnection", "Failed to connect to source control backend with the following errors:\n{0}\nThe source control backend had trouble connecting!\nTrying logging in with the 'p4 login' command or by using p4vs/UnrealGameSync."),
+			Log.Warning(FText::Format(LOCTEXT("FailedSourceControlConnection", "Failed to connect to revision control backend with the following errors:\n{0}\nThe revision control backend had trouble connecting!\nTrying logging in with the 'p4 login' command or by using p4vs/UnrealGameSync."),
 				Errors.ToText()));
 
 			OnConnectionError();
@@ -323,7 +323,7 @@ IVirtualizationBackend::EConnectionStatus FSourceControlBackend::OnConnect()
 	{
 		FMessageLog Log("LogVirtualization");
 
-		Log.Warning(FText::Format(LOCTEXT("FailedMetaInfo", "Failed to find 'payload_metainfo.txt' in the depot '{0}'\nThe source control backend will be unable to pull payloads, is your source control  config set up correctly?"),
+		Log.Warning(FText::Format(LOCTEXT("FailedMetaInfo", "Failed to find 'payload_metainfo.txt' in the depot '{0}'\nThe revision control backend will be unable to pull payloads, is your revision control config set up correctly?"),
 			FText::FromString(DepotRoot)));
 
 		OnConnectionError();
@@ -335,7 +335,7 @@ IVirtualizationBackend::EConnectionStatus FSourceControlBackend::OnConnect()
 	{
 		FMessageLog Log("LogVirtualization");
 
-		Log.Warning(FText::Format(LOCTEXT("FailedMetaInfo", "Failed to find 'payload_metainfo.txt' in the depot '{0}'\nThe source control backend will be unable to pull payloads, is your source control  config set up correctly?"),
+		Log.Warning(FText::Format(LOCTEXT("FailedMetaInfo", "Failed to find 'payload_metainfo.txt' in the depot '{0}'\nThe revision control backend will be unable to pull payloads, is your revision control config set up correctly?"),
 			FText::FromString(DepotRoot)));
 
 		OnConnectionError();
@@ -348,7 +348,7 @@ IVirtualizationBackend::EConnectionStatus FSourceControlBackend::OnConnect()
 	{
 		FMessageLog Log("LogVirtualization");
 
-		Log.Warning(FText::Format(LOCTEXT("FailedMetaInfo", "Failed to find 'payload_metainfo.txt' in the depot '{0}'\nThe source control backend will be unable to pull payloads, is your source control  config set up correctly?"),
+		Log.Warning(FText::Format(LOCTEXT("FailedMetaInfo", "Failed to find 'payload_metainfo.txt' in the depot '{0}'\nThe revision control backend will be unable to pull payloads, is your revision control config set up correctly?"),
 			FText::FromString(DepotRoot)));
 
 		OnConnectionError();
@@ -529,7 +529,7 @@ bool FSourceControlBackend::PushData(TArrayView<FPushRequest> Requests)
 		}
 	}
 
-	UE_LOG(LogVirtualization, Log, TEXT("[%s] Determined that %d payload(s) require submission to source control"), *GetDebugName(), RequestsToPush.Num());
+	UE_LOG(LogVirtualization, Log, TEXT("[%s] Determined that %d payload(s) require submission to revision control"), *GetDebugName(), RequestsToPush.Num());
 
 	if (RequestsToPush.IsEmpty())
 	{
@@ -600,7 +600,7 @@ bool FSourceControlBackend::PushData(TArrayView<FPushRequest> Requests)
 			CreateWorkspaceCommand->SetType(FCreateWorkspace::EType::Partitioned);
 		}
 
-		CreateWorkspaceCommand->SetDescription(TEXT("This workspace was autogenerated when submitting virtualized payloads to source control"));
+		CreateWorkspaceCommand->SetDescription(TEXT("This workspace was autogenerated when submitting virtualized payloads to revision control"));
 
 		if (SCCProvider->Execute(CreateWorkspaceCommand) != ECommandResult::Succeeded)
 		{
@@ -732,7 +732,7 @@ bool FSourceControlBackend::PushData(TArrayView<FPushRequest> Requests)
 
 			if (SCCProvider->Execute(ISourceControlOperation::Create<FMarkForAdd>(), FilesToSubmit) != ECommandResult::Succeeded)
 			{
-				UE_LOG(LogVirtualization, Error, TEXT("[%s] Failed to mark the payload file for Add in source control"), *GetDebugName());
+				UE_LOG(LogVirtualization, Error, TEXT("[%s] Failed to mark the payload file for Add in revision control"), *GetDebugName());
 				return false;
 			}
 		}
@@ -750,7 +750,7 @@ bool FSourceControlBackend::PushData(TArrayView<FPushRequest> Requests)
 
 			if (SCCProvider->Execute(CheckInOperation, FilesToSubmit) != ECommandResult::Succeeded)
 			{
-				UE_LOG(LogVirtualization, Error, TEXT("[%s] Failed to submit the payload file(s) to source control"), *GetDebugName());
+				UE_LOG(LogVirtualization, Error, TEXT("[%s] Failed to submit the payload file(s) to revision control"), *GetDebugName());
 				return false;
 			}
 		}
@@ -801,7 +801,7 @@ bool FSourceControlBackend::DoPayloadsExist(TArrayView<const FIoHash> PayloadIds
 	ECommandResult::Type Result = GetDepotPathStates(*SCCProvider, DepotPaths, PathStates);
 	if (Result != ECommandResult::Type::Succeeded)
 	{
-		UE_LOG(LogVirtualization, Error, TEXT("[%s] Failed to query the state of files in the source control depot"), *GetDebugName());
+		UE_LOG(LogVirtualization, Error, TEXT("[%s] Failed to query the state of files in the revision control depot"), *GetDebugName());
 		return false;
 	}
 
@@ -896,11 +896,11 @@ bool FSourceControlBackend::TryApplySettingsFromConfigFiles(const FString& Confi
 		if (MaxLimit != INDEX_NONE)
 		{
 			ConcurrentConnectionLimit = MakeUnique<FSemaphore>(MaxLimit);
-			UE_LOG(LogVirtualization, Log, TEXT("[%s] Limted to %d concurrent source control connections"), *GetDebugName(), MaxLimit);
+			UE_LOG(LogVirtualization, Log, TEXT("[%s] Limted to %d concurrent revision control connections"), *GetDebugName(), MaxLimit);
 		}
 		else
 		{
-			UE_LOG(LogVirtualization, Log, TEXT("[%s] Has no limit to it's concurrent source control connections"), *GetDebugName());
+			UE_LOG(LogVirtualization, Log, TEXT("[%s] Has no limit to it's concurrent revision control connections"), *GetDebugName());
 		}
 	}
 
diff --git a/Engine/Source/Editor/LevelEditor/Private/LevelEditorOutlinerSettings.cpp b/Engine/Source/Editor/LevelEditor/Private/LevelEditorOutlinerSettings.cpp
--- a/Engine/Source/Editor/LevelEditor/Private/LevelEditorOutlinerSettings.cpp
+++ b/Engine/Source/Editor/LevelEditor/Private/LevelEditorOutlinerSettings.cpp
@@ -324,7 +324,7 @@ void FLevelEditorOutlinerSettings::OnUnsavedAssetRemoved(const FString& InAsset)
 void FLevelEditorOutlinerSettings::CreateSCCFilters()
 {
 	// Source Control Category
-	TSharedPtr<FFilterCategory> SCCFiltersCategory = MakeShared<FFilterCategory>(LOCTEXT("SCCFiltersCategory", "Source Control"), FText::GetEmpty());
+	TSharedPtr<FFilterCategory> SCCFiltersCategory = MakeShared<FFilterCategory>(LOCTEXT("SCCFiltersCategory", "Revision Control"), FText::GetEmpty());
 	
 	// Uncontrolled Actors Filter
 	FUncontrolledChangelistsModule& UncontrolledChangelistModule = FUncontrolledChangelistsModule::Get();
@@ -333,7 +333,7 @@ void FLevelEditorOutlinerSettings::CreateSCCFilters()
 	
 	FGenericFilter<const ISceneOutlinerTreeItem&>::FOnItemFiltered UncontrolledFilterDelegate = FGenericFilter<const ISceneOutlinerTreeItem&>::FOnItemFiltered::CreateSP(this, &FLevelEditorOutlinerSettings::DoesActorPassUncontrolledFilter);
 	TSharedPtr<FGenericFilter<const ISceneOutlinerTreeItem&>> UncontrolledFilter = MakeShared<FGenericFilter<const ISceneOutlinerTreeItem&>>(SCCFiltersCategory, UncontrolledAssetsFilterName, LOCTEXT("UncontrolledFilterName", "Uncontrolled"), UncontrolledFilterDelegate);
-	UncontrolledFilter->SetToolTipText(LOCTEXT("UncontrolledFilterTooltip", "Only show items that are uncontrolled (locally modified outside of source control)"));
+	UncontrolledFilter->SetToolTipText(LOCTEXT("UncontrolledFilterTooltip", "Only show items that are uncontrolled (locally modified outside of revision control)"));
 	CustomFilters.Add(UncontrolledFilter.ToSharedRef());
 	
 	// File Management Category
diff --git a/Engine/Source/Editor/SceneOutliner/Private/SSourceControlWidget.cpp b/Engine/Source/Editor/SceneOutliner/Private/SSourceControlWidget.cpp
--- a/Engine/Source/Editor/SceneOutliner/Private/SSourceControlWidget.cpp
+++ b/Engine/Source/Editor/SceneOutliner/Private/SSourceControlWidget.cpp
@@ -72,7 +72,7 @@ void SSourceControlWidget::UpdateWidget(FSourceControlStatePtr SourceControlStat
 		SetColorAndOpacity(IconColor);
 
 		SetToolTipText(FText::Format(LOCTEXT("UncontrolledSourceControlCombinedTooltip",
-			"{0}\nStatus in source control: {1}"),
+			"{0}\nStatus in revision control: {1}"),
 			UncontrolledChangelistState->GetDisplayTooltip(), SourceControlState->GetDisplayTooltip()));
 	}
 	else if (SourceControlState.IsValid() && !UncontrolledChangelistState.IsValid())
diff --git a/Engine/Source/Editor/SourceControlWindows/Private/SSourceControlChangelists.cpp b/Engine/Source/Editor/SourceControlWindows/Private/SSourceControlChangelists.cpp
--- a/Engine/Source/Editor/SourceControlWindows/Private/SSourceControlChangelists.cpp
+++ b/Engine/Source/Editor/SourceControlWindows/Private/SSourceControlChangelists.cpp
@@ -872,7 +872,7 @@ void SSourceControlChangelistsWidget::OnRefreshSourceControlWidgets(int64 CurrUp
 
 	auto OnSourceControlCachedItemVisited = [CurrUpdateNum, &VisitedItemCount](IChangelistTreeItem* Item)
 	{
-		checkf(Item->VisitedUpdateNum < CurrUpdateNum, TEXT("The same source control item was visited twice. It is likely present in more than one changelist"));
+		checkf(Item->VisitedUpdateNum < CurrUpdateNum, TEXT("The same revision control item was visited twice. It is likely present in more than one changelist"));
 		Item->VisitedUpdateNum = CurrUpdateNum;
 		++VisitedItemCount;
 	};
