PR #9853: Add ability to filter FPrimaryAssetId by class. (Contributed by [REDACTED_PERSON])

diff --git a/Engine/Plugins/Editor/AssetManagerEditor/Source/AssetManagerEditor/Private/AssetManagerEditorModule.cpp b/Engine/Plugins/Editor/AssetManagerEditor/Source/AssetManagerEditor/Private/AssetManagerEditorModule.cpp
--- a/Engine/Plugins/Editor/AssetManagerEditor/Source/AssetManagerEditor/Private/AssetManagerEditorModule.cpp
+++ b/Engine/Plugins/Editor/AssetManagerEditor/Source/AssetManagerEditor/Private/AssetManagerEditorModule.cpp
@@ -7,7 +7,7 @@
 #include "PrimaryAssetIdCustomization.h"
 #include "SAssetAuditBrowser.h"
 #include "Engine/PrimaryAssetLabel.h"
-#include <Templates/UniquePtr.h>
+#include "Templates/UniquePtr.h"
 
 #include "Serialization/JsonReader.h"
 #include "Serialization/JsonSerializer.h"
@@ -78,7 +78,7 @@ DEFINE_LOG_CATEGORY(LogAssetManagerEditor);
 
 class FAssetManagerGraphPanelNodeFactory : public FGraphPanelNodeFactory
 {
-	virtual TSharedPtr<class SGraphNode> CreateNode(UEdGraphNode* Node) const override
+	virtual TSharedPtr<SGraphNode> CreateNode(UEdGraphNode* Node) const override
 	{
 		if (UEdGraphNode_Reference* DependencyNode = Cast<UEdGraphNode_Reference>(Node))
 		{
@@ -91,7 +91,7 @@ class FAssetManagerGraphPanelNodeFactory : public FGraphPanelNodeFactory
 
 class FAssetManagerGraphPanelPinFactory : public FGraphPanelPinFactory
 {
-	virtual TSharedPtr<class SGraphPin> CreatePin(class UEdGraphPin* InPin) const override
+	virtual TSharedPtr<SGraphPin> CreatePin(UEdGraphPin* InPin) const override
 	{
 		if (InPin->PinType.PinCategory == UEdGraphSchema_K2::PC_Struct && InPin->PinType.PinSubCategoryObject == TBaseStructure<FPrimaryAssetId>::Get())
 		{
@@ -136,9 +136,10 @@ TSharedRef<SWidget> IAssetManagerEditorModule::MakePrimaryAssetTypeSelector(FOnG
 	return PropertyCustomizationHelpers::MakePropertyComboBox(nullptr, GetStrings, GetValue, SetValue);
 }
 
-TSharedRef<SWidget> IAssetManagerEditorModule::MakePrimaryAssetIdSelector(FOnGetPrimaryAssetDisplayText OnGetDisplayText, FOnSetPrimaryAssetId OnSetId, bool bAllowClear, TArray<FPrimaryAssetType> AllowedTypes)
+TSharedRef<SWidget> IAssetManagerEditorModule::MakePrimaryAssetIdSelector(FOnGetPrimaryAssetDisplayText OnGetDisplayText, FOnSetPrimaryAssetId OnSetId, bool bAllowClear, 
+	TArray<FPrimaryAssetType> AllowedTypes, TArray<const UClass*> AllowedClasses, TArray<const UClass*> DisallowedClasses)
 {
-	FOnGetContent OnCreateMenuContent = FOnGetContent::CreateLambda([OnGetDisplayText, OnSetId, bAllowClear, AllowedTypes]()
+	FOnGetContent OnCreateMenuContent = FOnGetContent::CreateLambda([OnGetDisplayText, OnSetId, bAllowClear, AllowedTypes, AllowedClasses, DisallowedClasses]()
 	{
 		FOnShouldFilterAsset AssetFilter = FOnShouldFilterAsset::CreateStatic(&IAssetManagerEditorModule::OnShouldFilterPrimaryAsset, AllowedTypes);
 		FOnSetObject OnSetObject = FOnSetObject::CreateLambda([OnSetId](const FAssetData& AssetData)
@@ -156,13 +157,13 @@ TSharedRef<SWidget> IAssetManagerEditorModule::MakePrimaryAssetIdSelector(FOnGet
 			OnSetId.Execute(AssetId);
 		});
 
-		TArray<const UClass*> AllowedClasses;
 		TArray<UFactory*> NewAssetFactories;
 
 		return PropertyCustomizationHelpers::MakeAssetPickerWithMenu(
 			FAssetData(),
 			bAllowClear,
 			AllowedClasses,
+			DisallowedClasses,
 			NewAssetFactories,
 			AssetFilter,
 			OnSetObject,
@@ -183,7 +184,7 @@ TSharedRef<SWidget> IAssetManagerEditorModule::MakePrimaryAssetIdSelector(FOnGet
 		];
 }
 
-void IAssetManagerEditorModule::GeneratePrimaryAssetTypeComboBoxStrings(TArray< TSharedPtr<FString> >& OutComboBoxStrings, TArray<TSharedPtr<class SToolTip>>& OutToolTips, TArray<bool>& OutRestrictedItems, bool bAllowClear, bool bAllowAll)
+void IAssetManagerEditorModule::GeneratePrimaryAssetTypeComboBoxStrings(TArray< TSharedPtr<FString> >& OutComboBoxStrings, TArray<TSharedPtr<SToolTip>>& OutToolTips, TArray<bool>& OutRestrictedItems, bool bAllowClear, bool bAllowAll)
 {
 	UAssetManager& AssetManager = UAssetManager::Get();
 
diff --git a/Engine/Plugins/Editor/AssetManagerEditor/Source/AssetManagerEditor/Private/PrimaryAssetIdCustomization.cpp b/Engine/Plugins/Editor/AssetManagerEditor/Source/AssetManagerEditor/Private/PrimaryAssetIdCustomization.cpp
--- a/Engine/Plugins/Editor/AssetManagerEditor/Source/AssetManagerEditor/Private/PrimaryAssetIdCustomization.cpp
+++ b/Engine/Plugins/Editor/AssetManagerEditor/Source/AssetManagerEditor/Private/PrimaryAssetIdCustomization.cpp
@@ -101,6 +101,9 @@ void FPrimaryAssetIdCustomization::CustomizeHeader(TSharedRef<class IPropertyHan
 	// Can the field be cleared
 	const bool bAllowClear = !(StructPropertyHandle->GetMetaDataProperty()->PropertyFlags & CPF_NoClear);
 
+	AllowedClasses = PropertyCustomizationHelpers::GetClassesFromMetadataString(StructPropertyHandle->GetMetaData("AllowedClasses"));
+	DisallowedClasses = PropertyCustomizationHelpers::GetClassesFromMetadataString(StructPropertyHandle->GetMetaData("DisallowedClasses"));
+	
 	ValueBox->AddSlot()
 		.FillWidth(1.0f)
 		.VAlign(VAlign_Center)
@@ -108,7 +111,7 @@ void FPrimaryAssetIdCustomization::CustomizeHeader(TSharedRef<class IPropertyHan
 			IAssetManagerEditorModule::MakePrimaryAssetIdSelector(
 				FOnGetPrimaryAssetDisplayText::CreateSP(this, &FPrimaryAssetIdCustomization::GetDisplayText),
 				FOnSetPrimaryAssetId::CreateSP(this, &FPrimaryAssetIdCustomization::OnIdSelected),
-				bAllowClear, AllowedTypes)
+				bAllowClear, AllowedTypes, AllowedClasses, DisallowedClasses)
 		];
 
 	ValueBox->AddSlot()
@@ -131,6 +134,7 @@ void FPrimaryAssetIdCustomization::CustomizeHeader(TSharedRef<class IPropertyHan
 		[
 			PropertyCustomizationHelpers::MakeClearButton(FSimpleDelegate::CreateSP(this, &FPrimaryAssetIdCustomization::OnClear))
 		];
+
 	HeaderRow
 	.NameContent()
 	[
diff --git a/Engine/Plugins/Editor/AssetManagerEditor/Source/AssetManagerEditor/Private/PrimaryAssetIdCustomization.h b/Engine/Plugins/Editor/AssetManagerEditor/Source/AssetManagerEditor/Private/PrimaryAssetIdCustomization.h
--- a/Engine/Plugins/Editor/AssetManagerEditor/Source/AssetManagerEditor/Private/PrimaryAssetIdCustomization.h
+++ b/Engine/Plugins/Editor/AssetManagerEditor/Source/AssetManagerEditor/Private/PrimaryAssetIdCustomization.h
@@ -7,7 +7,11 @@
 #include "IPropertyTypeCustomization.h"
 #include "SGraphPin.h"
 
+class FAssetThumbnail;
+class FDetailWidgetRow;
+class IDetailChildrenBuilder;
 class IPropertyHandle;
+class SBorder;
 
 /** Customization for a primary asset id, shows an asset picker with filters */
 class FPrimaryAssetIdCustomization : public IPropertyTypeCustomization
@@ -19,8 +23,8 @@ public:
 	}
 
 	/** IPropertyTypeCustomization interface */
-	virtual void CustomizeHeader(TSharedRef<class IPropertyHandle> StructPropertyHandle, class FDetailWidgetRow& HeaderRow, IPropertyTypeCustomizationUtils& StructCustomizationUtils) override;
-	virtual void CustomizeChildren(TSharedRef<IPropertyHandle> InStructPropertyHandle, class IDetailChildrenBuilder& ChildBuilder, IPropertyTypeCustomizationUtils& StructCustomizationUtils) override {}
+	virtual void CustomizeHeader(TSharedRef<IPropertyHandle> StructPropertyHandle, FDetailWidgetRow& HeaderRow, IPropertyTypeCustomizationUtils& StructCustomizationUtils) override;
+	virtual void CustomizeChildren(TSharedRef<IPropertyHandle> InStructPropertyHandle, IDetailChildrenBuilder& ChildBuilder, IPropertyTypeCustomizationUtils& StructCustomizationUtils) override {}
 
 private:
 	void OnIdSelected(FPrimaryAssetId AssetId);
@@ -46,11 +50,17 @@ private:
 	/** Specified type */
 	TArray<FPrimaryAssetType> AllowedTypes;
 
+	/** Classes which can be selected with this PrimaryAssetId  */
+	TArray<const UClass*> AllowedClasses;
+
+	/** Classes which cannot be selected with this PrimaryAssetId  */
+	TArray<const UClass*> DisallowedClasses;
+	
 	/** Thumbnail resource */
-	TSharedPtr<class FAssetThumbnail> AssetThumbnail;
+	TSharedPtr<FAssetThumbnail> AssetThumbnail;
 
 	/** The border surrounding the thumbnail image. */
-	TSharedPtr< class SBorder > ThumbnailBorder;
+	TSharedPtr<SBorder> ThumbnailBorder;
 };
 
 /** Graph pin version of UI */
diff --git a/Engine/Plugins/Editor/AssetManagerEditor/Source/AssetManagerEditor/Public/AssetManagerEditorModule.h b/Engine/Plugins/Editor/AssetManagerEditor/Source/AssetManagerEditor/Public/AssetManagerEditorModule.h
--- a/Engine/Plugins/Editor/AssetManagerEditor/Source/AssetManagerEditor/Public/AssetManagerEditorModule.h
+++ b/Engine/Plugins/Editor/AssetManagerEditor/Source/AssetManagerEditor/Public/AssetManagerEditorModule.h
@@ -322,7 +322,9 @@ public:
 	 * @param OnSetId Delegate called when id is changed
 	 * @param bAllowClear If true, add None option to top
 	 */
-	static TSharedRef<SWidget> MakePrimaryAssetIdSelector(FOnGetPrimaryAssetDisplayText OnGetDisplayText, FOnSetPrimaryAssetId OnSetId, bool bAllowClear = true, TArray<FPrimaryAssetType> AllowedTypes = TArray<FPrimaryAssetType>());
+	static TSharedRef<SWidget> MakePrimaryAssetIdSelector(FOnGetPrimaryAssetDisplayText OnGetDisplayText, FOnSetPrimaryAssetId OnSetId, bool bAllowClear = true, 
+		TArray<FPrimaryAssetType> AllowedTypes = TArray<FPrimaryAssetType>(), TArray<const UClass*> AllowedClasses = TArray<const UClass*>(), TArray<const UClass*> DisallowedClasses = TArray<const UClass*>()
+	);
 
 	/** Called to get list of valid primary asset types */
 	static void GeneratePrimaryAssetTypeComboBoxStrings(TArray< TSharedPtr<FString> >& OutComboBoxStrings, TArray<TSharedPtr<SToolTip>>& OutToolTips, TArray<bool>& OutRestrictedItems, bool bAllowClear, bool bAllowAll);
