Notify Uncontrolled Changelists when files have their writeable status changed.
Cleanup Revert Writeable to use SourceControlHelpers::ApplyOperationAndReloadPackages utility instead of managing package reloading manually.

#tests Tested make and revert writeable. Tested revert before adding hook to make sure it was necessary

diff --git a/Engine/Plugins/Editor/AssetManagerEditor/Source/AssetManagerEditor/AssetManagerEditor.Build.cs b/Engine/Plugins/Editor/AssetManagerEditor/Source/AssetManagerEditor/AssetManagerEditor.Build.cs
--- a/Engine/Plugins/Editor/AssetManagerEditor/Source/AssetManagerEditor/AssetManagerEditor.Build.cs
+++ b/Engine/Plugins/Editor/AssetManagerEditor/Source/AssetManagerEditor/AssetManagerEditor.Build.cs
@@ -49,6 +49,7 @@ public class AssetManagerEditor : ModuleRules
 				"ToolWidgets",
 				"SourceControl",
 				"SourceControlWindows",
+				"UncontrolledChangelists",
 			}
 		);
 	}
diff --git a/Engine/Plugins/Editor/AssetManagerEditor/Source/AssetManagerEditor/Private/AssetSourceControlContextMenu.cpp b/Engine/Plugins/Editor/AssetManagerEditor/Source/AssetManagerEditor/Private/AssetSourceControlContextMenu.cpp
--- a/Engine/Plugins/Editor/AssetManagerEditor/Source/AssetManagerEditor/Private/AssetSourceControlContextMenu.cpp
+++ b/Engine/Plugins/Editor/AssetManagerEditor/Source/AssetManagerEditor/Private/AssetSourceControlContextMenu.cpp
@@ -24,6 +24,7 @@
 #include "Toolkits/AssetEditorToolkitMenuContext.h"
 #include "ToolMenu.h"
 #include "ToolMenuSection.h"
+#include "UncontrolledChangelistsModule.h"
 
 #include "Styling/CoreStyle.h"
 #include "Styling/ISlateStyle.h"
@@ -649,6 +650,7 @@ void FAssetSourceControlContextMenuState::ExecuteSCCSyncAndCheckOut() const
 void FAssetSourceControlContextMenuState::ExecuteSCCMakeWritable() const
 {
 	IPlatformFile& PlatformFile = FPlatformFileManager::Get().GetPlatformFile();
+	FUncontrolledChangelistsModule& UncontrolledChangelistsModule = FUncontrolledChangelistsModule::Get();
 
 	TArray<FString> PackageNames;
 	GetSelectedPackageNames(PackageNames);
@@ -656,6 +658,7 @@ void FAssetSourceControlContextMenuState::ExecuteSCCMakeWritable() const
 	for (const FString& Package : PackageNames)
 	{
 		PlatformFile.SetReadOnly(*SourceControlHelpers::PackageFilename(Package), false);
+		UncontrolledChangelistsModule.OnMakeWritable(Package);
 	}
 }
 
@@ -784,7 +787,6 @@ void FAssetSourceControlContextMenuState::ExecuteSCCRevertWritable() const
 	TArray<FString> PackageNames;
 	GetSelectedPackageNames(PackageNames);
 
-	TArray<UPackage*> PackagesToReload;
 	TArray<FString> PackageFileNamesToRevert;
 
 	// Only add packages that can actually be reverted
@@ -797,49 +799,49 @@ void FAssetSourceControlContextMenuState::ExecuteSCCRevertWritable() const
 			FSourceControlStatePtr SourceControlState = SourceControlProvider.GetState(PackageFileName, EStateCacheUsage::Use);
 			if (SourceControlState.IsValid() && SourceControlState->IsSourceControlled() && !SourceControlState->IsCheckedOut())
 			{
-				if (UPackage* Package = FindPackage(nullptr, *PackageName))
-				{
-					PackagesToReload.Add(Package);
-
-					// Detach the linkers of any loaded packages so that SCC can overwrite the files...
-					if (!Package->IsFullyLoaded())
-					{
-						FlushAsyncLoading();
-						Package->FullyLoad();
-					}
-					ResetLoaders(Package);
-				}
-
 				PackageFileNamesToRevert.Add(PackageFileName);
 			}
 		}
 	}
 
+	auto RevertOperation = [](const TArray<FString>& InFilenames) -> bool
 	{
-		TSharedRef<FUpdateStatus, ESPMode::ThreadSafe> UpdateOperation = ISourceControlOperation::Create<FUpdateStatus>();
-		UpdateOperation->SetUpdateHistory(true);
-		SourceControlProvider.Execute(UpdateOperation, PackageFileNamesToRevert);
-	}
+		ISourceControlProvider& SourceControlProvider = ISourceControlModule::Get().GetProvider();
+		IPlatformFile& PlatformFile = FPlatformFileManager::Get().GetPlatformFile();
 
-	for (const FString& PackageFileName : PackageFileNamesToRevert)
-	{
-		FSourceControlStatePtr SourceControlState = SourceControlProvider.GetState(PackageFileName, EStateCacheUsage::Use);
-		TSharedPtr<class ISourceControlRevision, ESPMode::ThreadSafe> Revision = SourceControlState->GetCurrentRevision();
+		// Get history for all files so we know the current revision to revert to
+		{
+			TSharedRef<FUpdateStatus, ESPMode::ThreadSafe> UpdateOperation = ISourceControlOperation::Create<FUpdateStatus>();
+			UpdateOperation->SetUpdateHistory(true);
+			SourceControlProvider.Execute(UpdateOperation, InFilenames);
+		}
 
-		if (Revision.IsValid())
+		// For sync files one at a time since we need to specify specific revisions on them
+		for (const FString& PackageFileName : InFilenames)
 		{
-			TSharedRef<FSync, ESPMode::ThreadSafe> SyncOperation = ISourceControlOperation::Create<FSync>();
-			SyncOperation->SetRevision(Revision->GetRevision());
-			SyncOperation->SetForce(true);
-			if (SourceControlProvider.Execute(SyncOperation, PackageFileName) == ECommandResult::Succeeded)
+			FSourceControlStatePtr SourceControlState = SourceControlProvider.GetState(PackageFileName, EStateCacheUsage::Use);
+			TSharedPtr<class ISourceControlRevision, ESPMode::ThreadSafe> Revision = SourceControlState->GetCurrentRevision();
+
+			if (Revision.IsValid())
 			{
-				PlatformFile.SetReadOnly(*PackageFileName, true);
+				TSharedRef<FSync, ESPMode::ThreadSafe> SyncOperation = ISourceControlOperation::Create<FSync>();
+				SyncOperation->SetRevision(Revision->GetRevision());
+				SyncOperation->SetForce(true);
+				if (SourceControlProvider.Execute(SyncOperation, PackageFileName) == ECommandResult::Succeeded)
+				{
+					PlatformFile.SetReadOnly(*PackageFileName, true);
+				}
 			}
 		}
-	}
 
-	// Hot-reload the new packages...
-	UPackageTools::ReloadPackages(PackagesToReload);
+		return true;
+	};
+
+	SourceControlHelpers::ApplyOperationAndReloadPackages(PackageNames, RevertOperation);
+
+	// Tell UncontrolledCL to refresh to pick up changes to files it was watching.
+	FUncontrolledChangelistsModule& UncontrolledChangelistsModule = FUncontrolledChangelistsModule::Get();
+	UncontrolledChangelistsModule.UpdateStatus();
 }
 
 void FAssetSourceControlContextMenuState::ExecuteSCCSync() const
