Fix compile error from merge

#tests compiled

diff --git a/Engine/Plugins/Editor/AssetManagerEditor/Source/AssetManagerEditor/Private/AssetSourceControlContextMenu.cpp b/Engine/Plugins/Editor/AssetManagerEditor/Source/AssetManagerEditor/Private/AssetSourceControlContextMenu.cpp
--- a/Engine/Plugins/Editor/AssetManagerEditor/Source/AssetManagerEditor/Private/AssetSourceControlContextMenu.cpp
+++ b/Engine/Plugins/Editor/AssetManagerEditor/Source/AssetManagerEditor/Private/AssetSourceControlContextMenu.cpp
@@ -848,7 +848,7 @@ void FAssetSourceControlContextMenuState::ExecuteSCCRevertWritable() const
 		return true;
 	};
 
-	SourceControlHelpers::ApplyOperationAndReloadPackages(PackageNames, RevertOperation);
+	SourceControlHelpers::ApplyOperationAndReloadPackages(PackageNames, PackageNames, RevertOperation);
 
 	// Tell UncontrolledCL to refresh to pick up changes to files it was watching.
 	FUncontrolledChangelistsModule& UncontrolledChangelistsModule = FUncontrolledChangelistsModule::Get();
