SizeMap
- Adding an option to filter out dependencies based on their type (game, editor only or all).
- Moved the size map window settings to a config so that it persists between sessions.

diff --git a/Engine/Plugins/Editor/AssetManagerEditor/Source/AssetManagerEditor/Private/SSizeMap.cpp b/Engine/Plugins/Editor/AssetManagerEditor/Source/AssetManagerEditor/Private/SSizeMap.cpp
--- a/Engine/Plugins/Editor/AssetManagerEditor/Source/AssetManagerEditor/Private/SSizeMap.cpp
+++ b/Engine/Plugins/Editor/AssetManagerEditor/Source/AssetManagerEditor/Private/SSizeMap.cpp
@@ -63,14 +63,18 @@ void SSizeMap::HandleSizeTypeComboChanged(TSharedPtr<FName> Item, ESelectInfo::T
 {
 	FSlateApplication::Get().DismissAllMenus();
 
-	CurrentSizeType = *Item.Get();
+	if (USizeMapSettings* SizeMapSettings = GetMutableDefault<USizeMapSettings>())
+	{
+		SizeMapSettings->SizeType = *Item.Get();
+		SizeMapSettings->SaveConfig();
+	}
 
 	RefreshMap();
 }
 
 FText SSizeMap::GetSizeTypeComboText() const
 {
-	return GetSizeTypeText(CurrentSizeType);
+	return GetSizeTypeText(GetDefault<USizeMapSettings>()->SizeType);
 }
 
 FText SSizeMap::GetSizeTypeText(FName SizeType) const
@@ -96,6 +100,42 @@ FText SSizeMap::GetOverviewText() const
 	return OverviewText;
 }
 
+TSharedRef<SWidget> SSizeMap::GenerateDependencyTypeComboItem(TSharedPtr<ESizeMapDependencyType> InItem) const
+{
+	return SNew(STextBlock).Text(GetDependencyTypeText(*InItem.Get()));
+}
+
+void SSizeMap::HandleDependencyTypeComboChanged(TSharedPtr<ESizeMapDependencyType> Item, ESelectInfo::Type SelectInfo)
+{
+	if (USizeMapSettings* SizeMapSettings = GetMutableDefault<USizeMapSettings>())
+	{
+		SizeMapSettings->DependencyType = *Item.Get();
+		SizeMapSettings->SaveConfig();
+	}
+
+	RefreshMap();
+}
+
+FText SSizeMap::GetDependencyTypeComboText() const
+{
+	return GetDependencyTypeText(GetDefault<USizeMapSettings>()->DependencyType);
+}
+
+FText SSizeMap::GetDependencyTypeText(ESizeMapDependencyType DependencyType) const
+{
+	switch (DependencyType)
+	{
+	case ESizeMapDependencyType::All:
+		return LOCTEXT("DependencyType_All", "All");
+	case ESizeMapDependencyType::Game:
+		return LOCTEXT("DependencyType_Game", "Game");
+	case ESizeMapDependencyType::EditorOnly:
+		return LOCTEXT("DependencyType_Editor", "Editor Only");
+	default:
+		return FText::GetEmpty();
+	}
+}
+
 void SSizeMap::Construct(const FArguments& InArgs)
 {
 	IAssetManagerEditorModule& ManagerEditorModule = IAssetManagerEditorModule::Get();
@@ -107,6 +147,10 @@ void SSizeMap::Construct(const FArguments& InArgs)
 	SizeTypeComboList.Add(MakeShared<FName>(IAssetManagerEditorModule::DiskSizeName));
 	SizeTypeComboList.Add(MakeShared<FName>(IAssetManagerEditorModule::ResourceSizeName));
 	
+	DependencyTypeComboList.Add(MakeShared<ESizeMapDependencyType>(ESizeMapDependencyType::All));
+	DependencyTypeComboList.Add(MakeShared<ESizeMapDependencyType>(ESizeMapDependencyType::Game));
+	DependencyTypeComboList.Add(MakeShared<ESizeMapDependencyType>(ESizeMapDependencyType::EditorOnly));
+
 	Commands = MakeShareable(new FUICommandList());
 
 	Commands->MapAction(
@@ -145,7 +189,6 @@ void SSizeMap::Construct(const FArguments& InArgs)
 		FAssetManagerEditorCommands::Get().MakeSharedCollectionWithDependencies,
 		FExecuteAction::CreateSP(this, &SSizeMap::MakeCollectionWithDependencies, ECollectionShareType::CST_Shared));
 
-	CurrentSizeType = IAssetManagerEditorModule::DiskSizeName;
 	CurrentRegistrySource = EditorModule->GetCurrentRegistrySource(false);
 	bMemorySizeCached = false;
 
@@ -184,6 +227,31 @@ void SSizeMap::Construct(const FArguments& InArgs)
 				.TextStyle(FAppStyle::Get(), "ContentBrowser.TopBar.Font")
 				.Text(this, &SSizeMap::GetOverviewText)
 			]
+			+ SHorizontalBox::Slot()
+			.HAlign(HAlign_Right)
+			.VAlign(VAlign_Center)
+			.Padding(2.f)
+			.AutoWidth()
+			[
+				SNew(STextBlock)
+				.Text(LOCTEXT("DependencyType_Label", "Dependencies to Display:"))
+			]
+			+ SHorizontalBox::Slot()
+			.HAlign(HAlign_Right)
+			.VAlign(VAlign_Center)
+			.Padding(2.f)
+			.AutoWidth()
+			[
+				SAssignNew(DependencyTypeComboBoxWidget, SComboBox<TSharedPtr<ESizeMapDependencyType>>)
+				.OptionsSource(&DependencyTypeComboList)
+				.OnGenerateWidget(this, &SSizeMap::GenerateDependencyTypeComboItem)
+				.OnSelectionChanged(this, &SSizeMap::HandleDependencyTypeComboChanged)
+				.ToolTipText(LOCTEXT("DependencyType_Tooltip", "Changes which dependencies are considered in the size calculation"))
+				[
+					SNew(STextBlock)
+					.Text(this, &SSizeMap::GetDependencyTypeComboText)
+				]
+			]
 			+SHorizontalBox::Slot()
 			.HAlign(HAlign_Right)
 			.VAlign(VAlign_Center)
@@ -199,7 +267,7 @@ void SSizeMap::Construct(const FArguments& InArgs)
 			.Padding(2.f)
 			.AutoWidth()
 			[
-				SAssignNew(ComboBoxWidget, SComboBox<TSharedPtr<FName>>)
+				SAssignNew(SizeTypeComboBoxWidget, SComboBox<TSharedPtr<FName>>)
 				.OptionsSource(&SizeTypeComboList)
 				.OnGenerateWidget(this, &SSizeMap::GenerateSizeTypeComboItem)
 				.OnSelectionChanged(this, &SSizeMap::HandleSizeTypeComboChanged)
@@ -283,8 +351,8 @@ void SSizeMap::SetCurrentRegistrySource(const FAssetManagerEditorRegistrySource*
 		if (!IsSizeTypeEnabled())
 		{
 			// If size type is disabled reset to disk size
-			CurrentSizeType = IAssetManagerEditorModule::DiskSizeName;
-			ComboBoxWidget->ClearSelection();
+			GetMutableDefault<USizeMapSettings>()->SizeType = IAssetManagerEditorModule::DiskSizeName;
+			SizeTypeComboBoxWidget->ClearSelection();
 		}
 	}
 	
@@ -539,6 +607,15 @@ void SSizeMap::GatherDependenciesRecursively(TSharedPtr<FAssetThumbnailPool>& In
 					DependencyQuery.Flags = UE::AssetRegistry::EDependencyQuery::Direct;
 				}
 				
+				if (GetDefault<USizeMapSettings>()->DependencyType == ESizeMapDependencyType::EditorOnly)
+				{
+					DependencyQuery.Flags |= UE::AssetRegistry::EDependencyQuery::EditorOnly;
+				}
+				else if (GetDefault<USizeMapSettings>()->DependencyType == ESizeMapDependencyType::Game)
+				{
+					DependencyQuery.Flags |= UE::AssetRegistry::EDependencyQuery::Game;
+				}
+
 				TArray<FAssetIdentifier> References;
 				
 				if (ChunkId != INDEX_NONE)
@@ -600,7 +677,7 @@ void SSizeMap::GatherDependenciesRecursively(TSharedPtr<FAssetThumbnailPool>& In
 
 				if (AssetPackageName != NAME_None)
 				{
-					if (EditorModule->GetIntegerValueForCustomColumn(NodeSizeMapData.AssetData, CurrentSizeType, FoundSize))
+					if (EditorModule->GetIntegerValueForCustomColumn(NodeSizeMapData.AssetData, GetDefault<USizeMapSettings>()->SizeType, FoundSize))
 					{
 						// If we're reading cooked data, this will fail for dependencies that are editor only. This is fine, they will have 0 size
 						NodeSizeMapData.AssetSize = FoundSize;
@@ -797,7 +874,7 @@ void SSizeMap::RefreshMap()
 	}
 
 	bool bShowSlowTask = false;
-	if (CurrentSizeType == IAssetManagerEditorModule::ResourceSizeName)
+	if (GetDefault<USizeMapSettings>()->SizeType == IAssetManagerEditorModule::ResourceSizeName)
 	{
 		if (!bMemorySizeCached)
 		{
@@ -806,7 +883,7 @@ void SSizeMap::RefreshMap()
 		}
 	}
 
-	// If we're refresing memory start a slow task as we need to laod assets
+	// If we're refreshing memory start a slow task as we need to load assets
 	FScopedSlowTask SlowTask(0, LOCTEXT("ComputeMemorySizeSlowTask", "Finding Memory Size..."), bShowSlowTask);
 	SlowTask.MakeDialog();
 
diff --git a/Engine/Plugins/Editor/AssetManagerEditor/Source/AssetManagerEditor/Private/SSizeMap.h b/Engine/Plugins/Editor/AssetManagerEditor/Source/AssetManagerEditor/Private/SSizeMap.h
--- a/Engine/Plugins/Editor/AssetManagerEditor/Source/AssetManagerEditor/Private/SSizeMap.h
+++ b/Engine/Plugins/Editor/AssetManagerEditor/Source/AssetManagerEditor/Private/SSizeMap.h
@@ -11,9 +11,35 @@
 #include "ITreeMap.h"
 #include "STreeMap.h"
 
+#include "SSizeMap.generated.h"
+
 class FAssetThumbnailPool;
 class FUICommandList;
 
+UENUM()
+enum class ESizeMapDependencyType
+{
+	/** Queries all hard dependencies */
+	All,
+	/** Queries hard dependencies that are part of a cooked build. See UE::AssetRegistry::EDependencyQuery::Game. */
+	Game,
+	/** Queries hard dependencies that only exist in the editor. See UE::AssetRegistry::EDependencyQuery::EditorOnly */
+	EditorOnly,
+};
+
+UCLASS(config=EditorPerProjectUserSettings)
+class USizeMapSettings : public UObject
+{
+	GENERATED_BODY()
+
+public:
+	UPROPERTY(config)
+	FName SizeType = IAssetManagerEditorModule::DiskSizeName;
+
+	UPROPERTY(config)
+	ESizeMapDependencyType DependencyType = ESizeMapDependencyType::All;
+};
+
 /**
  * Tree map for displaying the size of assets
  */
@@ -100,7 +126,7 @@ protected:
 	FReply OnZoomOut();
 	bool CanZoomOut() const;
 
-	/** Combo Box */
+	/** Size Type Combo Box */
 	TSharedRef<SWidget> GenerateSizeTypeComboItem(TSharedPtr<FName> InItem) const;
 	void HandleSizeTypeComboChanged(TSharedPtr<FName> Item, ESelectInfo::Type SelectInfo);
 	FText GetSizeTypeComboText() const;
@@ -109,10 +135,18 @@ protected:
 	bool IsSizeTypeEnabled() const;
 
 	TArray<TSharedPtr<FName>> SizeTypeComboList;
-	FName CurrentSizeType;
 	FText OverviewText;
 	bool bMemorySizeCached;
-	TSharedPtr<SComboBox<TSharedPtr<FName>>> ComboBoxWidget;
+	TSharedPtr<SComboBox<TSharedPtr<FName>>> SizeTypeComboBoxWidget;
+
+	/** Dependency Type Combo Box */
+	TSharedRef<SWidget> GenerateDependencyTypeComboItem(TSharedPtr<ESizeMapDependencyType> InItem) const;
+	void HandleDependencyTypeComboChanged(TSharedPtr<ESizeMapDependencyType> Item, ESelectInfo::Type SelectInfo);
+	FText GetDependencyTypeComboText() const;
+	FText GetDependencyTypeText(ESizeMapDependencyType DependencyType) const;
+
+	TArray<TSharedPtr<ESizeMapDependencyType>> DependencyTypeComboList;
+	TSharedPtr<SComboBox<TSharedPtr<ESizeMapDependencyType>>> DependencyTypeComboBoxWidget;
 
 	/** Our tree map widget */
 	TSharedPtr<class STreeMap> TreeMapWidget;
