#UE Collapsed reference viewer nodes in the referencer and dependency directions now have the identifiers of the assets they represent. This allows the tool to operate on each of those assets as a group, similar to how a collapsed node works when it is the root

diff --git a/Engine/Plugins/Editor/AssetManagerEditor/Source/AssetManagerEditor/Private/ReferenceViewer/EdGraphNode_Reference.cpp b/Engine/Plugins/Editor/AssetManagerEditor/Source/AssetManagerEditor/Private/ReferenceViewer/EdGraphNode_Reference.cpp
--- a/Engine/Plugins/Editor/AssetManagerEditor/Source/AssetManagerEditor/Private/ReferenceViewer/EdGraphNode_Reference.cpp
+++ b/Engine/Plugins/Editor/AssetManagerEditor/Source/AssetManagerEditor/Private/ReferenceViewer/EdGraphNode_Reference.cpp
@@ -114,12 +114,12 @@ void UEdGraphNode_Reference::SetupReferenceNode(const FIntPoint& NodeLoc, const
 	AllocateDefaultPins();
 }
 
-void UEdGraphNode_Reference::SetReferenceNodeCollapsed(const FIntPoint& NodeLoc, int32 InNumReferencesExceedingMax)
+void UEdGraphNode_Reference::SetReferenceNodeCollapsed(const FIntPoint& NodeLoc, int32 InNumReferencesExceedingMax, const TArray<FAssetIdentifier>& NewIdentifiers)
 {
 	NodePosX = NodeLoc.X;
 	NodePosY = NodeLoc.Y;
 
-	Identifiers.Empty();
+	Identifiers = NewIdentifiers;
 	bIsCollapsed = true;
 	bUsesThumbnail = false;
 	bIsOverflow = true;
diff --git a/Engine/Plugins/Editor/AssetManagerEditor/Source/AssetManagerEditor/Private/ReferenceViewer/EdGraph_ReferenceViewer.cpp b/Engine/Plugins/Editor/AssetManagerEditor/Source/AssetManagerEditor/Private/ReferenceViewer/EdGraph_ReferenceViewer.cpp
--- a/Engine/Plugins/Editor/AssetManagerEditor/Source/AssetManagerEditor/Private/ReferenceViewer/EdGraph_ReferenceViewer.cpp
+++ b/Engine/Plugins/Editor/AssetManagerEditor/Source/AssetManagerEditor/Private/ReferenceViewer/EdGraph_ReferenceViewer.cpp
@@ -684,9 +684,10 @@ UEdGraphNode_Reference* UEdGraph_ReferenceViewer::RecursivelyCreateNodes(bool bI
 		ChildLoc.Y -= (NodeProvSize - 1) * NodeSizeY * 0.5 ;
 
 		int32 Breadth = 0;
-
-		for (const TPair<FAssetIdentifier, EDependencyPinCategory>& Pair : InNodeInfos[InAssetId].Children)
+		int32 ChildIdx = 0;
+		for (; ChildIdx < InNodeInfos[InAssetId].Children.Num(); ChildIdx++)
 		{
+			const TPair<FAssetIdentifier, EDependencyPinCategory>& Pair = InNodeInfos[InAssetId].Children[ChildIdx];
 			if (ExceedsMaxSearchBreadth(Breadth) && !InNodeInfos[InAssetId].bExpandAllChildren)
 			{
 				break;
@@ -741,7 +742,14 @@ UEdGraphNode_Reference* UEdGraph_ReferenceViewer::RecursivelyCreateNodes(bool bI
 			if ( ensure(OverflowNode) )
 			{
 				OverflowNode->SetAllowThumbnail(!Settings->IsCompactMode());
-				OverflowNode->SetReferenceNodeCollapsed(RefNodeLoc, NodeInfo.OverflowCount);
+
+				TArray<FAssetIdentifier> CollapsedNodeIdentifiers;
+				for (; ChildIdx < InNodeInfos[InAssetId].Children.Num(); ChildIdx++)
+				{
+					const TPair<FAssetIdentifier, EDependencyPinCategory>& Pair = InNodeInfos[InAssetId].Children[ChildIdx];
+					CollapsedNodeIdentifiers.Add(Pair.Key);
+				}
+				OverflowNode->SetReferenceNodeCollapsed(RefNodeLoc, NodeInfo.OverflowCount, CollapsedNodeIdentifiers);
 
 				if ( bInReferencers )
 				{
diff --git a/Engine/Plugins/Editor/AssetManagerEditor/Source/AssetManagerEditor/Public/ReferenceViewer/EdGraphNode_Reference.h b/Engine/Plugins/Editor/AssetManagerEditor/Source/AssetManagerEditor/Public/ReferenceViewer/EdGraphNode_Reference.h
--- a/Engine/Plugins/Editor/AssetManagerEditor/Source/AssetManagerEditor/Public/ReferenceViewer/EdGraphNode_Reference.h
+++ b/Engine/Plugins/Editor/AssetManagerEditor/Source/AssetManagerEditor/Public/ReferenceViewer/EdGraphNode_Reference.h
@@ -61,7 +61,7 @@ class ASSETMANAGEREDITOR_API UEdGraphNode_Reference : public UEdGraphNode
 private:
 	void CacheAssetData(const FAssetData& AssetData);
 	void SetupReferenceNode(const FIntPoint& NodeLoc, const TArray<FAssetIdentifier>& NewIdentifiers, const FAssetData& InAssetData, bool bInAllowThumbnail, bool bInIsADuplicate);
-	void SetReferenceNodeCollapsed(const FIntPoint& NodeLoc, int32 InNumReferencesExceedingMax);
+	void SetReferenceNodeCollapsed(const FIntPoint& NodeLoc, int32 InNumReferencesExceedingMax, const TArray<FAssetIdentifier>& Identifiers);
 	void AddReferencer(class UEdGraphNode_Reference* ReferencerNode);
 
 	TArray<FAssetIdentifier> Identifiers;
