[Engine/Plugins/Editor]
* Ran IWYU on all private code in the plugins under this folder.

diff --git a/Engine/Plugins/Editor/AssetManagerEditor/Source/AssetManagerEditor/Private/AssetManagerEditorCommands.cpp b/Engine/Plugins/Editor/AssetManagerEditor/Source/AssetManagerEditor/Private/AssetManagerEditorCommands.cpp
--- a/Engine/Plugins/Editor/AssetManagerEditor/Source/AssetManagerEditor/Private/AssetManagerEditorCommands.cpp
+++ b/Engine/Plugins/Editor/AssetManagerEditor/Source/AssetManagerEditor/Private/AssetManagerEditorCommands.cpp
@@ -1,6 +1,7 @@
 // Copyright Epic Games, Inc. All Rights Reserved.
 
 #include "AssetManagerEditorCommands.h"
+#include "Styling/AppStyle.h"
 
 //////////////////////////////////////////////////////////////////////////
 // FAssetManagerEditorCommands
diff --git a/Engine/Plugins/Editor/AssetManagerEditor/Source/AssetManagerEditor/Private/AssetManagerEditorCommands.h b/Engine/Plugins/Editor/AssetManagerEditor/Source/AssetManagerEditor/Private/AssetManagerEditorCommands.h
--- a/Engine/Plugins/Editor/AssetManagerEditor/Source/AssetManagerEditor/Private/AssetManagerEditorCommands.h
+++ b/Engine/Plugins/Editor/AssetManagerEditor/Source/AssetManagerEditor/Private/AssetManagerEditorCommands.h
@@ -2,8 +2,6 @@
 
 #pragma once
 
-#include "CoreMinimal.h"
-#include "Styling/AppStyle.h"
 #include "Framework/Commands/Commands.h"
 
 // Actions that can be invoked in the reference viewer
diff --git a/Engine/Plugins/Editor/AssetManagerEditor/Source/AssetManagerEditor/Private/AssetManagerEditorModule.cpp b/Engine/Plugins/Editor/AssetManagerEditor/Source/AssetManagerEditor/Private/AssetManagerEditorModule.cpp
--- a/Engine/Plugins/Editor/AssetManagerEditor/Source/AssetManagerEditor/Private/AssetManagerEditorModule.cpp
+++ b/Engine/Plugins/Editor/AssetManagerEditor/Source/AssetManagerEditor/Private/AssetManagerEditorModule.cpp
@@ -1,46 +1,37 @@
 // Copyright Epic Games, Inc. All Rights Reserved.
 
 #include "AssetManagerEditorModule.h"
-#include "Modules/ModuleManager.h"
 #include "AssetRegistry/AssetRegistryModule.h"
+#include "ContentBrowserDataLegacyBridge.h"
+#include "HAL/PlatformFile.h"
 #include "PrimaryAssetTypeCustomization.h"
+#include "IDesktopPlatform.h"
 #include "PrimaryAssetIdCustomization.h"
+#include "Misc/CommandLine.h"
 #include "SAssetAuditBrowser.h"
-#include "Engine/PrimaryAssetLabel.h"
-#include "Templates/UniquePtr.h"
 
-#include "Serialization/JsonReader.h"
-#include "Serialization/JsonSerializer.h"
 #include "CollectionManagerModule.h"
-#include "GameDelegates.h"
 #include "ICollectionManager.h"
 #include "AssetRegistry/ARFilter.h"
 #include "Misc/FileHelper.h"
+#include "Misc/PackageName.h"
 #include "ProfilingDebugging/ProfilingHelpers.h"
-#include "Stats/StatsMisc.h"
 #include "Engine/AssetManager.h"
-#include "PropertyEditorModule.h"
-#include "IContentBrowserSingleton.h"
 #include "ContentBrowserModule.h"
+#include "ReferenceViewer/HistoryManager.h"
 #include "WorkspaceMenuStructure.h"
+#include "ToolMenu.h"
 #include "WorkspaceMenuStructureModule.h"
-#include "Framework/MultiBox/MultiBoxExtender.h"
-#include "Framework/MultiBox/MultiBoxBuilder.h"
-#include "Framework/Docking/TabManager.h"
 #include "Framework/Application/SlateApplication.h"
-#include "Framework/Commands/Commands.h"
+#include "ToolMenuEntry.h"
 #include "Widgets/Docking/SDockTab.h"
+#include "ToolMenuSection.h"
 #include "Widgets/Input/SComboButton.h"
 #include "Widgets/SToolTip.h"
 #include "PropertyCustomizationHelpers.h"
-#include "Toolkits/AssetEditorToolkit.h"
 
 #include "LevelEditor.h"
 #include "GraphEditorModule.h"
-#include "AssetRegistry/AssetData.h"
-#include "Engine/World.h"
-#include "Misc/App.h"
-#include "GenericPlatform/GenericPlatformFile.h"
 #include "Interfaces/ITargetPlatform.h"
 #include "Interfaces/ITargetPlatformManagerModule.h"
 #include "IPlatformFileSandboxWrapper.h"
@@ -48,24 +39,20 @@
 #include "Serialization/ArrayReader.h"
 #include "EdGraphUtilities.h"
 #include "EdGraphSchema_K2.h"
-#include "SGraphPin.h"
 #include "AssetManagerEditorCommands.h"
 #include "AssetSourceControlContextMenu.h"
 #include "ReferenceViewer/SReferenceViewer.h"
 #include "ReferenceViewer/SReferenceNode.h"
 #include "ReferenceViewer/EdGraphNode_Reference.h"
 #include "SSizeMap.h"
-#include "Editor.h"
 #include "Widgets/Notifications/SNotificationList.h"
 #include "Framework/Notifications/NotificationManager.h"
 #include "DesktopPlatformModule.h"
 #include "Misc/FileHelper.h"
 #include "Misc/MessageDialog.h"
-#include "UObject/ConstructorHelpers.h"
 #include "Misc/ScopedSlowTask.h"
 #include "IStatsViewer.h"
 #include "StatsViewerModule.h"
-#include "Subsystems/AssetEditorSubsystem.h"
 #include "ToolMenus.h"
 #include "Toolkits/AssetEditorToolkitMenuContext.h"
 #include "ContentBrowserMenuContexts.h"
diff --git a/Engine/Plugins/Editor/AssetManagerEditor/Source/AssetManagerEditor/Private/AssetSourceControlContextMenu.cpp b/Engine/Plugins/Editor/AssetManagerEditor/Source/AssetManagerEditor/Private/AssetSourceControlContextMenu.cpp
--- a/Engine/Plugins/Editor/AssetManagerEditor/Source/AssetManagerEditor/Private/AssetSourceControlContextMenu.cpp
+++ b/Engine/Plugins/Editor/AssetManagerEditor/Source/AssetManagerEditor/Private/AssetSourceControlContextMenu.cpp
@@ -4,39 +4,35 @@
 
 #include "AssetDefinition.h"
 #include "AssetDefinitionRegistry.h"
-#include "Algo/Transform.h"
+#include "AssetRegistry/ARFilter.h"
 #include "AssetRegistry/AssetRegistryModule.h"
 #include "AssetToolsModule.h"
 #include "AssetViewUtils.h"
 #include "Async/Async.h"
+#include "ContentBrowserDataMenuContexts.h"
 #include "ContentBrowserMenuContexts.h"
 #include "FileHelpers.h"
+#include "HAL/PlatformFile.h"
 #include "HAL/PlatformFileManager.h"
+#include "IAssetTools.h"
 #include "ISourceControlModule.h"
-#include "Misc/FileHelper.h"
+#include "ISourceControlRevision.h"
 #include "Misc/MessageDialog.h"
-#include "Misc/ScopedSlowTask.h"
+#include "Misc/PackageName.h"
 #include "SourceControlHelpers.h"
 #include "SourceControlOperations.h"
 #include "SourceControlWindows.h"
 #include "RevisionControlStyle/RevisionControlStyle.h"
-#include "Subsystems/AssetEditorSubsystem.h"
 #include "Toolkits/AssetEditorToolkit.h"
 #include "Toolkits/AssetEditorToolkitMenuContext.h"
 #include "ToolMenu.h"
 #include "ToolMenuSection.h"
 #include "UncontrolledChangelistsModule.h"
 
-#include "Styling/CoreStyle.h"
-#include "Styling/ISlateStyle.h"
-#include "Styling/ToolBarStyle.h"
-#include "ToolMenuDelegates.h"
 #include "ToolMenuEntry.h"
 #include "ToolMenuSection.h"
-#include "UObject/Linker.h"
 #include "Widgets/Images/SImage.h"
 #include "Widgets/Images/SThrobber.h"
-#include "Widgets/SBoxPanel.h"
 #include "Widgets/Text/STextBlock.h"
 
 #define LOCTEXT_NAMESPACE "AssetSourceControlContextMenu"
diff --git a/Engine/Plugins/Editor/AssetManagerEditor/Source/AssetManagerEditor/Private/AssetSourceControlContextMenu.h b/Engine/Plugins/Editor/AssetManagerEditor/Source/AssetManagerEditor/Private/AssetSourceControlContextMenu.h
--- a/Engine/Plugins/Editor/AssetManagerEditor/Source/AssetManagerEditor/Private/AssetSourceControlContextMenu.h
+++ b/Engine/Plugins/Editor/AssetManagerEditor/Source/AssetManagerEditor/Private/AssetSourceControlContextMenu.h
@@ -2,10 +2,7 @@
 
 #pragma once
 
-#include "CoreMinimal.h"
-#include "AssetRegistry/AssetData.h"
-#include "ISourceControlProvider.h"
-#include "HAL/CriticalSection.h"
+#include "Templates/SharedPointer.h"
 
 class UToolMenu;
 struct FToolMenuSection;
diff --git a/Engine/Plugins/Editor/AssetManagerEditor/Source/AssetManagerEditor/Private/PrimaryAssetIdCustomization.cpp b/Engine/Plugins/Editor/AssetManagerEditor/Source/AssetManagerEditor/Private/PrimaryAssetIdCustomization.cpp
--- a/Engine/Plugins/Editor/AssetManagerEditor/Source/AssetManagerEditor/Private/PrimaryAssetIdCustomization.cpp
+++ b/Engine/Plugins/Editor/AssetManagerEditor/Source/AssetManagerEditor/Private/PrimaryAssetIdCustomization.cpp
@@ -2,17 +2,17 @@
 
 #include "PrimaryAssetIdCustomization.h"
 #include "AssetManagerEditorModule.h"
-#include "Widgets/DeclarativeSyntaxSupport.h"
+#include "EdGraph/EdGraphSchema.h"
 #include "Widgets/Text/STextBlock.h"
 #include "Widgets/Input/SButton.h"
 #include "Widgets/Images/SImage.h"
 #include "Widgets/Layout/SBox.h"
 #include "Engine/AssetManager.h"
-#include "PropertyHandle.h"
 #include "PropertyCustomizationHelpers.h"
 #include "Editor.h"
 #include "AssetThumbnail.h"
 #include "ScopedTransaction.h"
+#include "Widgets/SOverlay.h"
 
 #define LOCTEXT_NAMESPACE "PrimaryAssetIdCustomization"
 
diff --git a/Engine/Plugins/Editor/AssetManagerEditor/Source/AssetManagerEditor/Private/PrimaryAssetIdCustomization.h b/Engine/Plugins/Editor/AssetManagerEditor/Source/AssetManagerEditor/Private/PrimaryAssetIdCustomization.h
--- a/Engine/Plugins/Editor/AssetManagerEditor/Source/AssetManagerEditor/Private/PrimaryAssetIdCustomization.h
+++ b/Engine/Plugins/Editor/AssetManagerEditor/Source/AssetManagerEditor/Private/PrimaryAssetIdCustomization.h
@@ -2,10 +2,9 @@
 
 #pragma once
 
-#include "CoreMinimal.h"
-#include "Widgets/SWidget.h"
 #include "IPropertyTypeCustomization.h"
 #include "SGraphPin.h"
+#include "UObject/PrimaryAssetId.h"
 
 class FAssetThumbnail;
 class FDetailWidgetRow;
diff --git a/Engine/Plugins/Editor/AssetManagerEditor/Source/AssetManagerEditor/Private/PrimaryAssetTypeCustomization.cpp b/Engine/Plugins/Editor/AssetManagerEditor/Source/AssetManagerEditor/Private/PrimaryAssetTypeCustomization.cpp
--- a/Engine/Plugins/Editor/AssetManagerEditor/Source/AssetManagerEditor/Private/PrimaryAssetTypeCustomization.cpp
+++ b/Engine/Plugins/Editor/AssetManagerEditor/Source/AssetManagerEditor/Private/PrimaryAssetTypeCustomization.cpp
@@ -2,11 +2,9 @@
 
 #include "PrimaryAssetTypeCustomization.h"
 #include "AssetManagerEditorModule.h"
-#include "Widgets/DeclarativeSyntaxSupport.h"
-#include "Widgets/SToolTip.h"
+#include "EdGraph/EdGraphSchema.h"
 #include "Widgets/Text/STextBlock.h"
 #include "Engine/AssetManager.h"
-#include "PropertyHandle.h"
 #include "PropertyCustomizationHelpers.h"
 #include "ScopedTransaction.h"
 
diff --git a/Engine/Plugins/Editor/AssetManagerEditor/Source/AssetManagerEditor/Private/PrimaryAssetTypeCustomization.h b/Engine/Plugins/Editor/AssetManagerEditor/Source/AssetManagerEditor/Private/PrimaryAssetTypeCustomization.h
--- a/Engine/Plugins/Editor/AssetManagerEditor/Source/AssetManagerEditor/Private/PrimaryAssetTypeCustomization.h
+++ b/Engine/Plugins/Editor/AssetManagerEditor/Source/AssetManagerEditor/Private/PrimaryAssetTypeCustomization.h
@@ -2,10 +2,9 @@
 
 #pragma once
 
-#include "CoreMinimal.h"
-#include "Widgets/SWidget.h"
 #include "IPropertyTypeCustomization.h"
 #include "SGraphPin.h"
+#include "UObject/PrimaryAssetId.h"
 
 class IPropertyHandle;
 
diff --git a/Engine/Plugins/Editor/AssetManagerEditor/Source/AssetManagerEditor/Private/ReferenceViewer/EdGraphNode_Reference.cpp b/Engine/Plugins/Editor/AssetManagerEditor/Source/AssetManagerEditor/Private/ReferenceViewer/EdGraphNode_Reference.cpp
--- a/Engine/Plugins/Editor/AssetManagerEditor/Source/AssetManagerEditor/Private/ReferenceViewer/EdGraphNode_Reference.cpp
+++ b/Engine/Plugins/Editor/AssetManagerEditor/Source/AssetManagerEditor/Private/ReferenceViewer/EdGraphNode_Reference.cpp
@@ -4,11 +4,13 @@
 #include "GenericPlatform/GenericPlatformFile.h"
 #include "EdGraph/EdGraphPin.h"
 #include "HAL/PlatformFileManager.h"
-#include "Styling/SlateIconFinder.h"
+#include "Misc/PackageName.h"
 #include "SReferenceViewer.h"
 #include "AssetToolsModule.h"
 #include "IAssetTools.h"
 #include "IAssetTypeActions.h"
+#include "ReferenceViewer/EdGraph_ReferenceViewer.h"
+#include "Styling/AppStyle.h"
 
 #include UE_INLINE_GENERATED_CPP_BY_NAME(EdGraphNode_Reference)
 
diff --git a/Engine/Plugins/Editor/AssetManagerEditor/Source/AssetManagerEditor/Private/ReferenceViewer/EdGraph_ReferenceViewer.cpp b/Engine/Plugins/Editor/AssetManagerEditor/Source/AssetManagerEditor/Private/ReferenceViewer/EdGraph_ReferenceViewer.cpp
--- a/Engine/Plugins/Editor/AssetManagerEditor/Source/AssetManagerEditor/Private/ReferenceViewer/EdGraph_ReferenceViewer.cpp
+++ b/Engine/Plugins/Editor/AssetManagerEditor/Source/AssetManagerEditor/Private/ReferenceViewer/EdGraph_ReferenceViewer.cpp
@@ -1,20 +1,18 @@
 // Copyright Epic Games, Inc. All Rights Reserved.
 
 #include "ReferenceViewer/EdGraph_ReferenceViewer.h"
+#include "Misc/FilterCollection.h"
 #include "ReferenceViewer/EdGraphNode_Reference.h"
+#include "Misc/IFilter.h"
 #include "ReferenceViewer/ReferenceViewerSettings.h"
 #include "EdGraph/EdGraphPin.h"
-#include "AssetRegistry/ARFilter.h"
 #include "AssetRegistry/AssetRegistryModule.h"
 #include "AssetThumbnail.h"
 #include "SReferenceViewer.h"
-#include "SReferenceNode.h"
-#include "GraphEditor.h"
 #include "ICollectionManager.h"
 #include "CollectionManagerModule.h"
 #include "AssetManagerEditorModule.h"
 #include "Engine/AssetManager.h"
-#include "Settings/EditorProjectSettings.h"
 
 #include UE_INLINE_GENERATED_CPP_BY_NAME(EdGraph_ReferenceViewer)
 
diff --git a/Engine/Plugins/Editor/AssetManagerEditor/Source/AssetManagerEditor/Private/ReferenceViewer/HistoryManager.cpp b/Engine/Plugins/Editor/AssetManagerEditor/Source/AssetManagerEditor/Private/ReferenceViewer/HistoryManager.cpp
--- a/Engine/Plugins/Editor/AssetManagerEditor/Source/AssetManagerEditor/Private/ReferenceViewer/HistoryManager.cpp
+++ b/Engine/Plugins/Editor/AssetManagerEditor/Source/AssetManagerEditor/Private/ReferenceViewer/HistoryManager.cpp
@@ -1,8 +1,7 @@
 // Copyright Epic Games, Inc. All Rights Reserved.
 
 #include "HistoryManager.h"
-#include "Textures/SlateIcon.h"
-#include "Framework/Commands/UIAction.h"
+#include "AssetRegistry/AssetIdentifier.h"
 #include "Framework/MultiBox/MultiBoxBuilder.h"
 
 #define LOCTEXT_NAMESPACE "ReferenceViewer"
diff --git a/Engine/Plugins/Editor/AssetManagerEditor/Source/AssetManagerEditor/Private/ReferenceViewer/HistoryManager.h b/Engine/Plugins/Editor/AssetManagerEditor/Source/AssetManagerEditor/Private/ReferenceViewer/HistoryManager.h
--- a/Engine/Plugins/Editor/AssetManagerEditor/Source/AssetManagerEditor/Private/ReferenceViewer/HistoryManager.h
+++ b/Engine/Plugins/Editor/AssetManagerEditor/Source/AssetManagerEditor/Private/ReferenceViewer/HistoryManager.h
@@ -2,8 +2,9 @@
 
 #pragma once
 
-#include "CoreMinimal.h"
-#include "AssetRegistry/AssetData.h"
+#include "Internationalization/Text.h"
+
+struct FAssetIdentifier;
 
 class FMenuBuilder;
 
diff --git a/Engine/Plugins/Editor/AssetManagerEditor/Source/AssetManagerEditor/Private/ReferenceViewer/ReferenceViewerSchema.cpp b/Engine/Plugins/Editor/AssetManagerEditor/Source/AssetManagerEditor/Private/ReferenceViewer/ReferenceViewerSchema.cpp
--- a/Engine/Plugins/Editor/AssetManagerEditor/Source/AssetManagerEditor/Private/ReferenceViewer/ReferenceViewerSchema.cpp
+++ b/Engine/Plugins/Editor/AssetManagerEditor/Source/AssetManagerEditor/Private/ReferenceViewer/ReferenceViewerSchema.cpp
@@ -1,17 +1,14 @@
 // Copyright Epic Games, Inc. All Rights Reserved.
 
 #include "ReferenceViewer/ReferenceViewerSchema.h"
-#include "Textures/SlateIcon.h"
-#include "Misc/Attribute.h"
 #include "SReferenceViewer.h"
-#include "ToolMenus.h"
-#include "EdGraph/EdGraph.h"
-#include "Styling/AppStyle.h"
-#include "CollectionManagerTypes.h"
 #include "AssetManagerEditorCommands.h"
 #include "AssetManagerEditorModule.h"
+#include "ToolMenu.h"
 #include "Toolkits/GlobalEditorCommonCommands.h"
 #include "ConnectionDrawingPolicy.h"
+#include "ToolMenuEntry.h"
+#include "ToolMenuSection.h"
 
 #include UE_INLINE_GENERATED_CPP_BY_NAME(ReferenceViewerSchema)
 
diff --git a/Engine/Plugins/Editor/AssetManagerEditor/Source/AssetManagerEditor/Private/ReferenceViewer/ReferenceViewerStyle.cpp b/Engine/Plugins/Editor/AssetManagerEditor/Source/AssetManagerEditor/Private/ReferenceViewer/ReferenceViewerStyle.cpp
--- a/Engine/Plugins/Editor/AssetManagerEditor/Source/AssetManagerEditor/Private/ReferenceViewer/ReferenceViewerStyle.cpp
+++ b/Engine/Plugins/Editor/AssetManagerEditor/Source/AssetManagerEditor/Private/ReferenceViewer/ReferenceViewerStyle.cpp
@@ -2,11 +2,14 @@
 
 #include "ReferenceViewerStyle.h"
 
+#include "Brushes/SlateImageBrush.h"
 #include "Styling/SlateTypes.h"
-#include "Styling/StarshipCoreStyle.h"
+#include "Brushes/SlateRoundedBoxBrush.h"
 #include "Styling/SlateStyleMacros.h"
+#include "Misc/Paths.h"
 #include "Styling/StyleColors.h"
 #include "Styling/CoreStyle.h"
+#include "Styling/SlateStyleRegistry.h"
 
 
 FName FReferenceViewerStyle::StyleName("ReferenceViewerStyle");
diff --git a/Engine/Plugins/Editor/AssetManagerEditor/Source/AssetManagerEditor/Private/ReferenceViewer/ReferenceViewerStyle.h b/Engine/Plugins/Editor/AssetManagerEditor/Source/AssetManagerEditor/Private/ReferenceViewer/ReferenceViewerStyle.h
--- a/Engine/Plugins/Editor/AssetManagerEditor/Source/AssetManagerEditor/Private/ReferenceViewer/ReferenceViewerStyle.h
+++ b/Engine/Plugins/Editor/AssetManagerEditor/Source/AssetManagerEditor/Private/ReferenceViewer/ReferenceViewerStyle.h
@@ -3,7 +3,6 @@
 #pragma once
 
 #include "Styling/SlateStyle.h"
-#include "Styling/SlateStyleRegistry.h"
 
 /**
  * The Reference Viewer is being used as a sandbox for testing out ui changes to graphs.
diff --git a/Engine/Plugins/Editor/AssetManagerEditor/Source/AssetManagerEditor/Private/ReferenceViewer/SReferenceNode.cpp b/Engine/Plugins/Editor/AssetManagerEditor/Source/AssetManagerEditor/Private/ReferenceViewer/SReferenceNode.cpp
--- a/Engine/Plugins/Editor/AssetManagerEditor/Source/AssetManagerEditor/Private/ReferenceViewer/SReferenceNode.cpp
+++ b/Engine/Plugins/Editor/AssetManagerEditor/Source/AssetManagerEditor/Private/ReferenceViewer/SReferenceNode.cpp
@@ -1,12 +1,13 @@
 // Copyright Epic Games, Inc. All Rights Reserved.
 
 #include "SReferenceNode.h"
+#include "ReferenceViewer/EdGraph_ReferenceViewer.h"
 #include "Widgets/SBoxPanel.h"
-#include "Widgets/Layout/SSpacer.h"
 #include "Widgets/Images/SImage.h"
 #include "Widgets/Layout/SBox.h"
 #include "ReferenceViewer/EdGraphNode_Reference.h"
 #include "AssetThumbnail.h"
+#include "Widgets/Notifications/SErrorText.h"
 #include "Widgets/Text/SInlineEditableTextBlock.h"
 #include "SCommentBubble.h"
 #include "ReferenceViewerStyle.h"
diff --git a/Engine/Plugins/Editor/AssetManagerEditor/Source/AssetManagerEditor/Private/ReferenceViewer/SReferenceNode.h b/Engine/Plugins/Editor/AssetManagerEditor/Source/AssetManagerEditor/Private/ReferenceViewer/SReferenceNode.h
--- a/Engine/Plugins/Editor/AssetManagerEditor/Source/AssetManagerEditor/Private/ReferenceViewer/SReferenceNode.h
+++ b/Engine/Plugins/Editor/AssetManagerEditor/Source/AssetManagerEditor/Private/ReferenceViewer/SReferenceNode.h
@@ -2,8 +2,6 @@
 
 #pragma once
 
-#include "CoreMinimal.h"
-#include "Widgets/DeclarativeSyntaxSupport.h"
 #include "SGraphNode.h"
 
 class FAssetThumbnail;
diff --git a/Engine/Plugins/Editor/AssetManagerEditor/Source/AssetManagerEditor/Private/ReferenceViewer/SReferenceViewer.cpp b/Engine/Plugins/Editor/AssetManagerEditor/Source/AssetManagerEditor/Private/ReferenceViewer/SReferenceViewer.cpp
--- a/Engine/Plugins/Editor/AssetManagerEditor/Source/AssetManagerEditor/Private/ReferenceViewer/SReferenceViewer.cpp
+++ b/Engine/Plugins/Editor/AssetManagerEditor/Source/AssetManagerEditor/Private/ReferenceViewer/SReferenceViewer.cpp
@@ -1,28 +1,16 @@
 // Copyright Epic Games, Inc. All Rights Reserved.
 
 #include "SReferenceViewer.h"
-#include "Widgets/SOverlay.h"
-#include "Engine/GameViewportClient.h"
 #include "Dialogs/Dialogs.h"
-#include "Misc/MessageDialog.h"
+#include "Framework/Views/TableViewMetadata.h"
 #include "Misc/ScopedSlowTask.h"
-#include "Modules/ModuleManager.h"
-#include "Widgets/SBoxPanel.h"
-#include "Widgets/Layout/SBorder.h"
-#include "Widgets/Images/SImage.h"
-#include "Widgets/Text/STextBlock.h"
-#include "Widgets/Layout/SBox.h"
-#include "Widgets/Input/SEditableTextBox.h"
-#include "Widgets/Input/SButton.h"
-#include "Widgets/Input/SCheckBox.h"
+#include "Misc/PackageName.h"
 #include "Widgets/Input/SSearchBox.h"
+#include "ReferenceViewer/EdGraph_ReferenceViewer.h"
 #include "Widgets/Input/SSpinBox.h"
-#include "SSimpleButton.h"
-#include "SSimpleComboButton.h"
 
-#include "Styling/AppStyle.h"
+#include "ReferenceViewer/HistoryManager.h"
 #include "ReferenceViewerStyle.h"
-#include "Engine/Selection.h"
 #include "ReferenceViewer/EdGraphNode_Reference.h"
 #include "ReferenceViewer/ReferenceViewerSchema.h"
 #include "AssetRegistry/AssetRegistryModule.h"
@@ -33,15 +21,16 @@
 #include "Editor.h"
 #include "AssetManagerEditorCommands.h"
 #include "EditorWidgetsModule.h"
+#include "ReferenceViewer/ReferenceViewerSettings.h"
 #include "Toolkits/GlobalEditorCommonCommands.h"
 #include "Engine/AssetManager.h"
+#include "ReferenceViewer/SReferenceViewerFilterBar.h"
 #include "Widgets/Input/SComboBox.h"
 #include "HAL/PlatformApplicationMisc.h"
 #include "AssetManagerEditorModule.h"
-#include "Framework/Application/SlateApplication.h"
-#include "Framework/MultiBox/MultiBoxBuilder.h"
 
 #include "ObjectTools.h"
+#include "Selection.h"
 
 #define LOCTEXT_NAMESPACE "ReferenceViewer"
 
@@ -2232,4 +2221,4 @@ void SReferenceViewer::OnFindPathAssetEnterPressed( const TArray<FAssetData>& As
 	RegisterActiveTimer(0.1f, FWidgetActiveTimerDelegate::CreateSP(this, &SReferenceViewer::TriggerZoomToFit));
 } 
 
-#undef LOCTEXT_NAMESPACE
\ No newline at end of file
+#undef LOCTEXT_NAMESPACE
diff --git a/Engine/Plugins/Editor/AssetManagerEditor/Source/AssetManagerEditor/Private/ReferenceViewer/SReferenceViewer.h b/Engine/Plugins/Editor/AssetManagerEditor/Source/AssetManagerEditor/Private/ReferenceViewer/SReferenceViewer.h
--- a/Engine/Plugins/Editor/AssetManagerEditor/Source/AssetManagerEditor/Private/ReferenceViewer/SReferenceViewer.h
+++ b/Engine/Plugins/Editor/AssetManagerEditor/Source/AssetManagerEditor/Private/ReferenceViewer/SReferenceViewer.h
@@ -2,21 +2,19 @@
 
 #pragma once
 
-#include "CoreMinimal.h"
-#include "Input/Reply.h"
-#include "Widgets/DeclarativeSyntaxSupport.h"
-#include "Widgets/SCompoundWidget.h"
-#include "Widgets/Input/SComboBox.h"
-#include "Widgets/Input/SSearchBox.h"
 #include "GraphEditor.h"
 #include "AssetRegistry/AssetData.h"
 #include "HistoryManager.h"
 #include "CollectionManagerTypes.h"
-#include "AssetManagerEditorModule.h"
-#include "Containers/ArrayView.h"
-#include "ReferenceViewer/ReferenceViewerSettings.h"
-#include "ReferenceViewer/EdGraph_ReferenceViewer.h"
-#include "ReferenceViewer/SReferenceViewerFilterBar.h"
+
+class FUICommandList;
+class SReferenceViewerFilterBar;
+class SSearchBox;
+class UEdGraph_ReferenceViewer;
+class UReferenceViewerSettings;
+namespace ESelectInfo { enum Type : int; }
+struct FAssetManagerEditorRegistrySource;
+template <typename OptionType> class SComboBox;
 
 class UEdGraph;
 
diff --git a/Engine/Plugins/Editor/AssetManagerEditor/Source/AssetManagerEditor/Private/ReferenceViewer/SReferenceViewerFilterBar.cpp b/Engine/Plugins/Editor/AssetManagerEditor/Source/AssetManagerEditor/Private/ReferenceViewer/SReferenceViewerFilterBar.cpp
--- a/Engine/Plugins/Editor/AssetManagerEditor/Source/AssetManagerEditor/Private/ReferenceViewer/SReferenceViewerFilterBar.cpp
+++ b/Engine/Plugins/Editor/AssetManagerEditor/Source/AssetManagerEditor/Private/ReferenceViewer/SReferenceViewerFilterBar.cpp
@@ -2,7 +2,7 @@
 
 #include "ReferenceViewer/SReferenceViewerFilterBar.h"
 
-#include "CoreMinimal.h"
+#include "ReferenceViewer/ReferenceViewerSettings.h"
 
 void SReferenceViewerFilterBar::SaveSettings()
 {
diff --git a/Engine/Plugins/Editor/AssetManagerEditor/Source/AssetManagerEditor/Private/ReferenceViewer/SReferenceViewerFilterBar.h b/Engine/Plugins/Editor/AssetManagerEditor/Source/AssetManagerEditor/Private/ReferenceViewer/SReferenceViewerFilterBar.h
--- a/Engine/Plugins/Editor/AssetManagerEditor/Source/AssetManagerEditor/Private/ReferenceViewer/SReferenceViewerFilterBar.h
+++ b/Engine/Plugins/Editor/AssetManagerEditor/Source/AssetManagerEditor/Private/ReferenceViewer/SReferenceViewerFilterBar.h
@@ -2,10 +2,9 @@
 
 #pragma once
 
-#include "CoreMinimal.h"
 #include "Filters/SFilterBar.h"
-#include "ReferenceViewer/ReferenceViewerSettings.h"
-#include "ReferenceViewer/EdGraph_ReferenceViewer.h"
+
+struct FReferenceNodeInfo;
 
 class SReferenceViewerFilterBar : public SFilterBar< FReferenceNodeInfo& > 
 {
@@ -18,4 +17,4 @@ public:
 	/** Loads any settings to config that should be persistent between editor sessions */
 	void LoadSettings() override;
 
-};
\ No newline at end of file
+};
diff --git a/Engine/Plugins/Editor/AssetManagerEditor/Source/AssetManagerEditor/Private/SAssetAuditBrowser.cpp b/Engine/Plugins/Editor/AssetManagerEditor/Source/AssetManagerEditor/Private/SAssetAuditBrowser.cpp
--- a/Engine/Plugins/Editor/AssetManagerEditor/Source/AssetManagerEditor/Private/SAssetAuditBrowser.cpp
+++ b/Engine/Plugins/Editor/AssetManagerEditor/Source/AssetManagerEditor/Private/SAssetAuditBrowser.cpp
@@ -2,35 +2,23 @@
 
 
 #include "SAssetAuditBrowser.h"
-#include "Framework/Commands/UIAction.h"
+#include "AssetManagerEditorModule.h"
 #include "Framework/Commands/UICommandList.h"
-#include "Widgets/Layout/SBorder.h"
-#include "Widgets/SBoxPanel.h"
-#include "Widgets/Text/STextBlock.h"
-#include "Widgets/Input/SMenuAnchor.h"
 #include "Framework/MultiBox/MultiBoxBuilder.h"
 
-#include "Styling/AppStyle.h"
-#include "Widgets/Layout/SBox.h"
+#include "Framework/Views/TableViewMetadata.h"
 #include "Widgets/Input/SButton.h"
-#include "Widgets/SViewport.h"
 #include "FileHelpers.h"
-#include "AssetRegistry/ARFilter.h"
 #include "ClassViewerModule.h"
-#include "ClassViewerFilter.h"
 #include "IContentBrowserSingleton.h"
 #include "ContentBrowserModule.h"
 #include "ContentBrowserDataSource.h"
 #include "AssetRegistry/AssetRegistryModule.h"
 #include "Toolkits/GlobalEditorCommonCommands.h"
-#include "FrontendFilterBase.h"
-#include "Slate/SceneViewport.h"
-#include "ObjectEditorUtils.h"
 #include "Engine/AssetManager.h"
 #include "AssetManagerEditorCommands.h"
 #include "Engine/BlueprintCore.h"
 #include "Widgets/Input/SComboBox.h"
-#include "Framework/Application/SlateApplication.h"
 #include "DragAndDrop/AssetDragDropOp.h"
 #include "Blueprint/BlueprintSupport.h"
 #include "Subsystems/AssetEditorSubsystem.h"
diff --git a/Engine/Plugins/Editor/AssetManagerEditor/Source/AssetManagerEditor/Private/SAssetAuditBrowser.h b/Engine/Plugins/Editor/AssetManagerEditor/Source/AssetManagerEditor/Private/SAssetAuditBrowser.h
--- a/Engine/Plugins/Editor/AssetManagerEditor/Source/AssetManagerEditor/Private/SAssetAuditBrowser.h
+++ b/Engine/Plugins/Editor/AssetManagerEditor/Source/AssetManagerEditor/Private/SAssetAuditBrowser.h
@@ -2,15 +2,11 @@
 
 #pragma once
 
-#include "CoreMinimal.h"
-#include "Layout/Visibility.h"
-#include "Widgets/DeclarativeSyntaxSupport.h"
-#include "Input/Reply.h"
-#include "Widgets/SWidget.h"
-#include "AssetRegistry/AssetData.h"
-#include "Widgets/SToolTip.h"
 #include "ContentBrowserDelegates.h"
-#include "AssetManagerEditorModule.h"
+#include "Widgets/SCompoundWidget.h"
+
+struct FAssetIdentifier;
+struct FAssetManagerEditorRegistrySource;
 
 class FUICommandList;
 class SMenuAnchor;
diff --git a/Engine/Plugins/Editor/AssetManagerEditor/Source/AssetManagerEditor/Private/SSizeMap.cpp b/Engine/Plugins/Editor/AssetManagerEditor/Source/AssetManagerEditor/Private/SSizeMap.cpp
--- a/Engine/Plugins/Editor/AssetManagerEditor/Source/AssetManagerEditor/Private/SSizeMap.cpp
+++ b/Engine/Plugins/Editor/AssetManagerEditor/Source/AssetManagerEditor/Private/SSizeMap.cpp
@@ -1,18 +1,16 @@
 // Copyright Epic Games, Inc. All Rights Reserved.
 
 #include "SSizeMap.h"
-#include "Modules/ModuleManager.h"
 #include "Engine/AssetManager.h"
 #include "Editor.h"
 #include "AssetRegistry/AssetRegistryModule.h"
 #include "AssetThumbnail.h"
 #include "ClassIconFinder.h"
-#include "Math/UnitConversion.h"
-#include "Widgets/Text/STextBlock.h"
+#include "Framework/Views/TableViewMetadata.h"
 #include "Widgets/Input/SButton.h"
-#include "Framework/Commands/UIAction.h"
 #include "Framework/Commands/UICommandList.h"
 
+#include "Misc/PackageName.h"
 #include "Toolkits/GlobalEditorCommonCommands.h"
 #include "Framework/MultiBox/MultiBoxBuilder.h"
 #include "AssetManagerEditorModule.h"
@@ -20,9 +18,10 @@
 #include "DragAndDrop/AssetDragDropOp.h"
 #include "CollectionManagerModule.h"
 #include "ICollectionManager.h"
-#include "Framework/Application/SlateApplication.h"
 #include "Misc/ScopedSlowTask.h"
+#include "STreeMap.h"
 #include "Subsystems/AssetEditorSubsystem.h"
+#include "Widgets/Input/SComboBox.h"
 
 #define LOCTEXT_NAMESPACE "SizeMap"
 
diff --git a/Engine/Plugins/Editor/AssetManagerEditor/Source/AssetManagerEditor/Private/SSizeMap.h b/Engine/Plugins/Editor/AssetManagerEditor/Source/AssetManagerEditor/Private/SSizeMap.h
--- a/Engine/Plugins/Editor/AssetManagerEditor/Source/AssetManagerEditor/Private/SSizeMap.h
+++ b/Engine/Plugins/Editor/AssetManagerEditor/Source/AssetManagerEditor/Private/SSizeMap.h
@@ -2,17 +2,16 @@
 
 #pragma once
 
-#include "CoreMinimal.h"
-#include "Widgets/DeclarativeSyntaxSupport.h"
+#include "CollectionManagerTypes.h"
 #include "Widgets/SCompoundWidget.h"
-#include "Widgets/Input/SComboBox.h"
 #include "AssetRegistry/AssetData.h"
-#include "AssetManagerEditorModule.h"
-#include "ITreeMap.h"
-#include "STreeMap.h"
-
 #include "SSizeMap.generated.h"
 
+class FTreeMapNodeData;
+namespace ESelectInfo { enum Type : int; }
+struct FAssetIdentifier;
+struct FAssetManagerEditorRegistrySource;
+template <typename OptionType> class SComboBox;
 class FAssetThumbnailPool;
 class FUICommandList;
 
diff --git a/Engine/Plugins/Editor/AssetReferenceRestrictions/Source/AssetReferenceRestrictions/Private/AssetReferencingDomains.cpp b/Engine/Plugins/Editor/AssetReferenceRestrictions/Source/AssetReferenceRestrictions/Private/AssetReferencingDomains.cpp
--- a/Engine/Plugins/Editor/AssetReferenceRestrictions/Source/AssetReferenceRestrictions/Private/AssetReferencingDomains.cpp
+++ b/Engine/Plugins/Editor/AssetReferenceRestrictions/Source/AssetReferenceRestrictions/Private/AssetReferencingDomains.cpp
@@ -5,9 +5,10 @@
 #include "Interfaces/IPluginManager.h"
 #include "DomainAssetReferenceFilter.h"
 #include "Editor.h"
+#include "Misc/PackageName.h"
 #include "Settings/ProjectPackagingSettings.h"
-#include "Misc/PackagePath.h"
-#include "Misc/StringBuilder.h"
+#include "Misc/Paths.h"
+#include "PluginReferenceDescriptor.h"
 
 DEFINE_LOG_CATEGORY(LogAssetReferenceRestrictions);
 
@@ -526,4 +527,4 @@ void FDomainDatabase::AddDomainVisibilityList(TSharedPtr<FDomainData> Domain, co
 	}
 }
 
-#undef LOCTEXT_NAMESPACE
\ No newline at end of file
+#undef LOCTEXT_NAMESPACE
diff --git a/Engine/Plugins/Editor/AssetReferenceRestrictions/Source/AssetReferenceRestrictions/Private/AssetReferencingDomains.h b/Engine/Plugins/Editor/AssetReferenceRestrictions/Source/AssetReferenceRestrictions/Private/AssetReferencingDomains.h
--- a/Engine/Plugins/Editor/AssetReferenceRestrictions/Source/AssetReferenceRestrictions/Private/AssetReferencingDomains.h
+++ b/Engine/Plugins/Editor/AssetReferenceRestrictions/Source/AssetReferenceRestrictions/Private/AssetReferencingDomains.h
@@ -2,8 +2,7 @@
 
 #pragma once
 
-#include "CoreMinimal.h"
-#include "Templates/Tuple.h"
+#include "Internationalization/Text.h"
 
 //@TODO: Debug printing
 #define UE_ASSET_DOMAIN_FILTERING_DEBUG_LOGGING 0
@@ -109,4 +108,4 @@ private:
 	TArray<FString> DomainsDefinedByPlugins;
 
 	bool bDatabaseOutOfDate = false;
-};
\ No newline at end of file
+};
diff --git a/Engine/Plugins/Editor/AssetReferenceRestrictions/Source/AssetReferenceRestrictions/Private/AssetReferencingPolicySettings.h b/Engine/Plugins/Editor/AssetReferenceRestrictions/Source/AssetReferenceRestrictions/Private/AssetReferencingPolicySettings.h
--- a/Engine/Plugins/Editor/AssetReferenceRestrictions/Source/AssetReferenceRestrictions/Private/AssetReferencingPolicySettings.h
+++ b/Engine/Plugins/Editor/AssetReferenceRestrictions/Source/AssetReferenceRestrictions/Private/AssetReferencingPolicySettings.h
@@ -2,11 +2,13 @@
 
 #pragma once
 
-#include "CoreMinimal.h"
 #include "Engine/DeveloperSettings.h"
 
 #include "AssetReferencingPolicySettings.generated.h"
 
+struct FDirectoryPath;
+struct FPropertyChangedEvent;
+
 USTRUCT()
 struct FARPDefaultPluginDomainRules
 {
diff --git a/Engine/Plugins/Editor/AssetReferenceRestrictions/Source/AssetReferenceRestrictions/Private/AssetReferencingPolicySubsystem.cpp b/Engine/Plugins/Editor/AssetReferenceRestrictions/Source/AssetReferenceRestrictions/Private/AssetReferencingPolicySubsystem.cpp
--- a/Engine/Plugins/Editor/AssetReferenceRestrictions/Source/AssetReferenceRestrictions/Private/AssetReferencingPolicySubsystem.cpp
+++ b/Engine/Plugins/Editor/AssetReferenceRestrictions/Source/AssetReferenceRestrictions/Private/AssetReferencingPolicySubsystem.cpp
@@ -1,7 +1,6 @@
 // Copyright Epic Games, Inc. All Rights Reserved.
 
 #include "AssetReferencingPolicySubsystem.h"
-#include "AssetReferencingPolicySettings.h"
 #include "DomainAssetReferenceFilter.h"
 #include "Editor.h"
 #include "AssetReferencingDomains.h"
diff --git a/Engine/Plugins/Editor/AssetReferenceRestrictions/Source/AssetReferenceRestrictions/Private/AssetValidator_AssetReferenceRestrictions.cpp b/Engine/Plugins/Editor/AssetReferenceRestrictions/Source/AssetReferenceRestrictions/Private/AssetValidator_AssetReferenceRestrictions.cpp
--- a/Engine/Plugins/Editor/AssetReferenceRestrictions/Source/AssetReferenceRestrictions/Private/AssetValidator_AssetReferenceRestrictions.cpp
+++ b/Engine/Plugins/Editor/AssetReferenceRestrictions/Source/AssetReferenceRestrictions/Private/AssetValidator_AssetReferenceRestrictions.cpp
@@ -1,14 +1,13 @@
 // Copyright Epic Games, Inc. All Rights Reserved.
 
 #include "AssetValidator_AssetReferenceRestrictions.h"
-#include "AssetRegistry/AssetData.h"
 #include "AssetRegistry/AssetRegistryModule.h"
-#include "UnrealEdGlobals.h"
 #include "Editor.h"
-#include "Editor/UnrealEdEngine.h"
 #include "AssetReferencingPolicySubsystem.h"
 #include "AssetReferencingDomains.h"
-#include "AssetReferencingPolicySettings.h"
+#include "Editor/AssetReferenceFilter.h"
+#include "Misc/PackageName.h"
+#include "Modules/ModuleManager.h"
 
 #include UE_INLINE_GENERATED_CPP_BY_NAME(AssetValidator_AssetReferenceRestrictions)
 
diff --git a/Engine/Plugins/Editor/AssetReferenceRestrictions/Source/AssetReferenceRestrictions/Private/DomainAssetReferenceFilter.cpp b/Engine/Plugins/Editor/AssetReferenceRestrictions/Source/AssetReferenceRestrictions/Private/DomainAssetReferenceFilter.cpp
--- a/Engine/Plugins/Editor/AssetReferenceRestrictions/Source/AssetReferenceRestrictions/Private/DomainAssetReferenceFilter.cpp
+++ b/Engine/Plugins/Editor/AssetReferenceRestrictions/Source/AssetReferenceRestrictions/Private/DomainAssetReferenceFilter.cpp
@@ -1,10 +1,10 @@
 // Copyright Epic Games, Inc. All Rights Reserved.
 
 #include "DomainAssetReferenceFilter.h"
+#include "Editor/EditorEngine.h"
 #include "Toolkits/ToolkitManager.h"
+#include "Modules/ModuleManager.h"
 #include "Toolkits/IToolkit.h"
-#include "Engine/AssetManager.h"
-#include "Interfaces/IPluginManager.h"
 #include "AssetRegistry/AssetRegistryModule.h"
 #include "AssetReferencingDomains.h"
 
diff --git a/Engine/Plugins/Editor/AssetReferenceRestrictions/Source/AssetReferenceRestrictions/Private/DomainAssetReferenceFilter.h b/Engine/Plugins/Editor/AssetReferenceRestrictions/Source/AssetReferenceRestrictions/Private/DomainAssetReferenceFilter.h
--- a/Engine/Plugins/Editor/AssetReferenceRestrictions/Source/AssetReferenceRestrictions/Private/DomainAssetReferenceFilter.h
+++ b/Engine/Plugins/Editor/AssetReferenceRestrictions/Source/AssetReferenceRestrictions/Private/DomainAssetReferenceFilter.h
@@ -2,7 +2,10 @@
 
 #pragma once
 
-#include "Editor/UnrealEdEngine.h"
+#include "AssetRegistry/AssetData.h"
+#include "Editor/AssetReferenceFilter.h"
+
+struct FAssetReferenceFilterContext;
 
 struct FAssetData;
 struct FDomainDatabase;
diff --git a/Engine/Plugins/Editor/AssetSearch/Source/Private/AssetSearchDatabase.cpp b/Engine/Plugins/Editor/AssetSearch/Source/Private/AssetSearchDatabase.cpp
--- a/Engine/Plugins/Editor/AssetSearch/Source/Private/AssetSearchDatabase.cpp
+++ b/Engine/Plugins/Editor/AssetSearch/Source/Private/AssetSearchDatabase.cpp
@@ -4,19 +4,13 @@
 
 #include "SQLiteDatabase.h"
 #include "HAL/FileManager.h"
-#include "UObject/StructOnScope.h"
-#include "Serialization/MemoryReader.h"
-#include "Serialization/MemoryWriter.h"
 
 #include "AssetRegistry/AssetData.h"
 #include "Misc/TextFilterExpressionEvaluator.h"
-#include "Misc/PackageName.h"
-#include "Misc/Paths.h"
-#include "Misc/SecureHash.h"
-#include "Engine/World.h"
-#include "HAL/PlatformFileManager.h"
-#include "Misc/StringBuilder.h"
 #include "Algo/LevenshteinDistance.h"
+#include "SQLitePreparedStatement.h"
+#include "SearchQuery.h"
+#include "Serialization/JsonReader.h"
 
 DEFINE_LOG_CATEGORY(LogAssetSearch);
 
diff --git a/Engine/Plugins/Editor/AssetSearch/Source/Private/AssetSearchDatabase.h b/Engine/Plugins/Editor/AssetSearch/Source/Private/AssetSearchDatabase.h
--- a/Engine/Plugins/Editor/AssetSearch/Source/Private/AssetSearchDatabase.h
+++ b/Engine/Plugins/Editor/AssetSearch/Source/Private/AssetSearchDatabase.h
@@ -2,8 +2,10 @@
 
 #pragma once
 
-#include "CoreMinimal.h"
-#include "IAssetSearchModule.h"
+#include "Logging/LogMacros.h"
+#include "Templates/UniquePtr.h"
+
+struct FSearchRecord;
 
 class FSQLiteDatabase;
 enum class ESQLiteDatabaseOpenMode : uint8;
@@ -111,4 +113,4 @@ private:
 
 	/** Prepared statements for the currently open database */
 	TUniquePtr<FAssetSearchDatabaseStatements> Statements;
-};
\ No newline at end of file
+};
diff --git a/Engine/Plugins/Editor/AssetSearch/Source/Private/AssetSearchManager.cpp b/Engine/Plugins/Editor/AssetSearch/Source/Private/AssetSearchManager.cpp
--- a/Engine/Plugins/Editor/AssetSearch/Source/Private/AssetSearchManager.cpp
+++ b/Engine/Plugins/Editor/AssetSearch/Source/Private/AssetSearchManager.cpp
@@ -2,35 +2,27 @@
 
 #include "AssetSearchManager.h"
 
-#include "AssetRegistry/IAssetRegistry.h"
 #include "AssetRegistry/AssetRegistryModule.h"
-#include "AssetSearchDatabase.h"
-#include "Async/Async.h"
 #include "DerivedDataCacheInterface.h"
 
-#include "Containers/StringConv.h"
+#include "IAssetSearchModule.h"
 #include "Misc/Paths.h"
 #include "HAL/RunnableThread.h"
+#include "Misc/App.h"
 #include "StudioAnalytics.h"
 #include "AnalyticsEventAttribute.h"
-#include "Misc/FeedbackContext.h"
+#include "Misc/PackageName.h"
 #include "WidgetBlueprint.h"
-#include "Engine/Blueprint.h"
 #include "Engine/DataTable.h"
 #include "Engine/DataAsset.h"
+#include "SearchSerializer.h"
 #include "Sound/DialogueWave.h"
 #include "Settings/SearchProjectSettings.h"
 #include "Settings/SearchUserSettings.h"
 #include "Sound/SoundCue.h"
 #include "Misc/ScopedSlowTask.h"
-#include "Misc/StringBuilder.h"
-#include "Engine/World.h"
 #include "Editor.h"
-#include "PackageTools.h"
-#include "UObject/ObjectMacros.h"
-#include "UObject/ObjectKey.h"
 #include "UObject/ObjectSaveContext.h"
-#include "UObject/WeakObjectPtrTemplates.h"
 #include "FileHelpers.h"
 #include "Misc/MessageDialog.h"
 #include "Engine/CurveTable.h"
@@ -38,8 +30,6 @@
 #include "Materials/MaterialFunction.h"
 #include "Materials/MaterialParameterCollection.h"
 #include "Materials/MaterialInstance.h"
-#include "ISearchProvider.h"
-#include "Stats/Stats.h"
 
 #include "Indexers/GenericObjectIndexer.h"
 #include "Indexers/DataTableIndexer.h"
diff --git a/Engine/Plugins/Editor/AssetSearch/Source/Private/AssetSearchManager.h b/Engine/Plugins/Editor/AssetSearch/Source/Private/AssetSearchManager.h
--- a/Engine/Plugins/Editor/AssetSearch/Source/Private/AssetSearchManager.h
+++ b/Engine/Plugins/Editor/AssetSearch/Source/Private/AssetSearchManager.h
@@ -2,13 +2,17 @@
 
 #pragma once
 
-#include "CoreMinimal.h"
-#include "IAssetSearchModule.h"
+#include "AssetRegistry/AssetData.h"
 #include "AssetSearchDatabase.h"
 #include "FileInfoDatabase.h"
 #include "Containers/Queue.h"
 #include "Containers/Ticker.h"
 #include "HAL/Runnable.h"
+#include "SearchQuery.h"
+#include "UObject/WeakObjectPtr.h"
+
+class IAssetIndexer;
+struct FSearchStats;
 
 class FObjectPostSaveContext;
 class FRunnableThread;
@@ -122,4 +126,4 @@ private:
 	TQueue<TFunction<void()>, EQueueMode::Mpsc> ImmediateOperations;
 	TQueue<TFunction<void()>, EQueueMode::Mpsc> FeedOperations;
 	TQueue<TFunction<void()>, EQueueMode::Mpsc> UpdateOperations;
-};
\ No newline at end of file
+};
diff --git a/Engine/Plugins/Editor/AssetSearch/Source/Private/AssetSearchModule.cpp b/Engine/Plugins/Editor/AssetSearch/Source/Private/AssetSearchModule.cpp
--- a/Engine/Plugins/Editor/AssetSearch/Source/Private/AssetSearchModule.cpp
+++ b/Engine/Plugins/Editor/AssetSearch/Source/Private/AssetSearchModule.cpp
@@ -1,15 +1,10 @@
 // Copyright Epic Games, Inc. All Rights Reserved.
 
-#include "CoreMinimal.h"
-#include "Modules/ModuleManager.h"
 #include "IAssetSearchModule.h"
-#include "Widgets/DeclarativeSyntaxSupport.h"
 #include "WorkspaceMenuStructure.h"
-#include "Styling/AppStyle.h"
 #include "WorkspaceMenuStructureModule.h"
 #include "Widgets/Docking/SDockTab.h"
 #include "Widgets/SSearchBrowser.h"
-#include "AssetRegistry/AssetRegistryModule.h"
 #include "AssetSearchManager.h"
 
 #define LOCTEXT_NAMESPACE "FAssetSearchModule"
diff --git a/Engine/Plugins/Editor/AssetSearch/Source/Private/FileInfoDatabase.cpp b/Engine/Plugins/Editor/AssetSearch/Source/Private/FileInfoDatabase.cpp
--- a/Engine/Plugins/Editor/AssetSearch/Source/Private/FileInfoDatabase.cpp
+++ b/Engine/Plugins/Editor/AssetSearch/Source/Private/FileInfoDatabase.cpp
@@ -2,20 +2,15 @@
 
 #include "FileInfoDatabase.h"
 
+#include "HAL/PlatformFile.h"
 #include "SQLiteDatabase.h"
 #include "HAL/FileManager.h"
-#include "UObject/StructOnScope.h"
-#include "Serialization/MemoryReader.h"
-#include "Serialization/MemoryWriter.h"
 
 #include "AssetRegistry/AssetData.h"
-#include "Misc/TextFilterExpressionEvaluator.h"
 #include "Misc/PackageName.h"
 #include "Misc/Paths.h"
-#include "Misc/SecureHash.h"
-#include "Engine/World.h"
 #include "HAL/PlatformFileManager.h"
-#include "Misc/StringBuilder.h"
+#include "SQLitePreparedStatement.h"
 
 DECLARE_LOG_CATEGORY_CLASS(LogFileInfo, Log, All);
 
@@ -452,4 +447,4 @@ TMap<FName, FAssetFileInfo> FFileInfoDatabase::GetAllFileInfos()
 	return FileInfos;
 }
 
-//"database disk image is malformed"
\ No newline at end of file
+//"database disk image is malformed"
diff --git a/Engine/Plugins/Editor/AssetSearch/Source/Private/FileInfoDatabase.h b/Engine/Plugins/Editor/AssetSearch/Source/Private/FileInfoDatabase.h
--- a/Engine/Plugins/Editor/AssetSearch/Source/Private/FileInfoDatabase.h
+++ b/Engine/Plugins/Editor/AssetSearch/Source/Private/FileInfoDatabase.h
@@ -2,8 +2,8 @@
 
 #pragma once
 
-#include "CoreMinimal.h"
-#include "IAssetSearchModule.h"
+#include "Misc/DateTime.h"
+#include "Misc/SecureHash.h"
 
 class FSQLiteDatabase;
 enum class ESQLiteDatabaseOpenMode : uint8;
@@ -96,4 +96,4 @@ private:
 
 	/** Prepared statements for the currently open database */
 	TUniquePtr<FFileInfoDatabaseStatements> Statements;
-};
\ No newline at end of file
+};
diff --git a/Engine/Plugins/Editor/AssetSearch/Source/Private/Indexers/ActorIndexer.h b/Engine/Plugins/Editor/AssetSearch/Source/Private/Indexers/ActorIndexer.h
--- a/Engine/Plugins/Editor/AssetSearch/Source/Private/Indexers/ActorIndexer.h
+++ b/Engine/Plugins/Editor/AssetSearch/Source/Private/Indexers/ActorIndexer.h
@@ -2,12 +2,14 @@
 
 #pragma once
 
-#include "CoreMinimal.h"
+#include "Containers/UnrealString.h"
 #include "IAssetIndexer.h"
 
+struct FScriptContainerElement;
+
 class FActorIndexer : public IAssetIndexer
 {
 	virtual FString GetName() const override { return TEXT("Actor"); }
 	virtual int32 GetVersion() const override;
 	virtual void IndexAsset(const UObject* InAssetObject, FSearchSerializer& Serializer) const override;
-};
\ No newline at end of file
+};
diff --git a/Engine/Plugins/Editor/AssetSearch/Source/Private/Indexers/BlueprintIndexer.cpp b/Engine/Plugins/Editor/AssetSearch/Source/Private/Indexers/BlueprintIndexer.cpp
--- a/Engine/Plugins/Editor/AssetSearch/Source/Private/Indexers/BlueprintIndexer.cpp
+++ b/Engine/Plugins/Editor/AssetSearch/Source/Private/Indexers/BlueprintIndexer.cpp
@@ -1,17 +1,13 @@
 // Copyright Epic Games, Inc. All Rights Reserved.
 
 #include "BlueprintIndexer.h"
-#include "Engine/Blueprint.h"
-#include "EdGraph/EdGraphPin.h"
+#include "EdGraph/EdGraph.h"
 #include "K2Node_CallFunction.h"
 #include "K2Node_Variable.h"
-#include "Kismet2/BlueprintEditorUtils.h"
 #include "Utility/IndexerUtilities.h"
 #include "K2Node_BaseMCDelegate.h"
-#include "Internationalization/Text.h"
 #include "K2Node_Knot.h"
 #include "EdGraphNode_Comment.h"
-#include "Engine/SimpleConstructionScript.h"
 #include "Engine/SCS_Node.h"
 #include "K2Node_Event.h"
 #include "SearchSerializer.h"
@@ -179,4 +175,4 @@ void FBlueprintIndexer::IndexMemberReference(FSearchSerializer& Serializer, cons
 	}
 }
 
-#undef LOCTEXT_NAMESPACE
\ No newline at end of file
+#undef LOCTEXT_NAMESPACE
diff --git a/Engine/Plugins/Editor/AssetSearch/Source/Private/Indexers/BlueprintIndexer.h b/Engine/Plugins/Editor/AssetSearch/Source/Private/Indexers/BlueprintIndexer.h
--- a/Engine/Plugins/Editor/AssetSearch/Source/Private/Indexers/BlueprintIndexer.h
+++ b/Engine/Plugins/Editor/AssetSearch/Source/Private/Indexers/BlueprintIndexer.h
@@ -2,9 +2,11 @@
 
 #pragma once
 
-#include "CoreMinimal.h"
+#include "Containers/UnrealString.h"
 #include "IAssetIndexer.h"
 
+struct FScriptContainerElement;
+
 class UBlueprint;
 struct FMemberReference;
 class FSearchSerializer;
@@ -22,4 +24,4 @@ private:
 	void IndexComponents(const UBlueprint* InBlueprint, FSearchSerializer& Serializer) const;
 	void IndexGraphs(const UBlueprint* InBlueprint, FSearchSerializer& Serializer) const;
 	void IndexMemberReference(FSearchSerializer& Serializer, const FMemberReference& MemberReference, const FString& MemberType) const;
-};
\ No newline at end of file
+};
diff --git a/Engine/Plugins/Editor/AssetSearch/Source/Private/Indexers/CurveTableIndexer.cpp b/Engine/Plugins/Editor/AssetSearch/Source/Private/Indexers/CurveTableIndexer.cpp
--- a/Engine/Plugins/Editor/AssetSearch/Source/Private/Indexers/CurveTableIndexer.cpp
+++ b/Engine/Plugins/Editor/AssetSearch/Source/Private/Indexers/CurveTableIndexer.cpp
@@ -1,7 +1,6 @@
 // Copyright Epic Games, Inc. All Rights Reserved.
 
 #include "CurveTableIndexer.h"
-#include "Utility/IndexerUtilities.h"
 #include "Engine/CurveTable.h"
 #include "SearchSerializer.h"
 
@@ -33,4 +32,4 @@ void FCurveTableIndexer::IndexAsset(const UObject* InAssetObject, FSearchSeriali
 		Serializer.IndexProperty(TEXT("Row"), RowName);
 	}
 	Serializer.EndIndexingObject();
-}
\ No newline at end of file
+}
diff --git a/Engine/Plugins/Editor/AssetSearch/Source/Private/Indexers/DialogueWaveIndexer.cpp b/Engine/Plugins/Editor/AssetSearch/Source/Private/Indexers/DialogueWaveIndexer.cpp
--- a/Engine/Plugins/Editor/AssetSearch/Source/Private/Indexers/DialogueWaveIndexer.cpp
+++ b/Engine/Plugins/Editor/AssetSearch/Source/Private/Indexers/DialogueWaveIndexer.cpp
@@ -1,8 +1,6 @@
 // Copyright Epic Games, Inc. All Rights Reserved.
 
 #include "DialogueWaveIndexer.h"
-#include "Utility/IndexerUtilities.h"
-#include "Engine/DataAsset.h"
 #include "Sound/DialogueWave.h"
 #include "SearchSerializer.h"
 
@@ -30,4 +28,4 @@ void FDialogueWaveIndexer::IndexAsset(const UObject* InAssetObject, FSearchSeria
 	Serializer.IndexProperty(TEXT("SpokenText"), DialogueWave->SpokenText);
 
 	Serializer.EndIndexingObject();
-}
\ No newline at end of file
+}
diff --git a/Engine/Plugins/Editor/AssetSearch/Source/Private/Indexers/LevelIndexer.cpp b/Engine/Plugins/Editor/AssetSearch/Source/Private/Indexers/LevelIndexer.cpp
--- a/Engine/Plugins/Editor/AssetSearch/Source/Private/Indexers/LevelIndexer.cpp
+++ b/Engine/Plugins/Editor/AssetSearch/Source/Private/Indexers/LevelIndexer.cpp
@@ -2,7 +2,6 @@
 
 #include "LevelIndexer.h"
 #include "Utility/IndexerUtilities.h"
-#include "Engine/DataAsset.h"
 #include "Engine/LevelScriptBlueprint.h"
 #include "Engine/World.h"
 #include "SearchSerializer.h"
@@ -48,4 +47,4 @@ void FLevelIndexer::IndexAsset(const UObject* InAssetObject, FSearchSerializer&
 			}
 		}
 	}
-}
\ No newline at end of file
+}
diff --git a/Engine/Plugins/Editor/AssetSearch/Source/Private/Indexers/MaterialExpressionIndexer.cpp b/Engine/Plugins/Editor/AssetSearch/Source/Private/Indexers/MaterialExpressionIndexer.cpp
--- a/Engine/Plugins/Editor/AssetSearch/Source/Private/Indexers/MaterialExpressionIndexer.cpp
+++ b/Engine/Plugins/Editor/AssetSearch/Source/Private/Indexers/MaterialExpressionIndexer.cpp
@@ -2,20 +2,11 @@
 
 #include "MaterialExpressionIndexer.h"
 #include "Utility/IndexerUtilities.h"
-#include "Internationalization/Text.h"
 #include "Materials/Material.h"
 #include "Materials/MaterialFunction.h"
-#include "Materials/MaterialExpression.h"
 #include "Materials/MaterialExpressionFunctionOutput.h"
 #include "SearchSerializer.h"
 
-#include "Materials/MaterialExpressionParameter.h"
-#include "Materials/MaterialExpressionFontSampleParameter.h"
-#include "Materials/MaterialExpressionScalarParameter.h"
-#include "Materials/MaterialExpressionStaticBoolParameter.h"
-#include "Materials/MaterialExpressionStaticComponentMaskParameter.h"
-#include "Materials/MaterialExpressionTextureSampleParameter.h"
-#include "Materials/MaterialExpressionVectorParameter.h"
 
 #define LOCTEXT_NAMESPACE "FMaterialExpressionIndexer"
 
@@ -69,4 +60,4 @@ void FMaterialExpressionIndexer::IndexParameters(const UObject* InAssetObject, F
 	}
 }
 
-#undef LOCTEXT_NAMESPACE
\ No newline at end of file
+#undef LOCTEXT_NAMESPACE
diff --git a/Engine/Plugins/Editor/AssetSearch/Source/Private/Indexers/SoundCueIndexer.cpp b/Engine/Plugins/Editor/AssetSearch/Source/Private/Indexers/SoundCueIndexer.cpp
--- a/Engine/Plugins/Editor/AssetSearch/Source/Private/Indexers/SoundCueIndexer.cpp
+++ b/Engine/Plugins/Editor/AssetSearch/Source/Private/Indexers/SoundCueIndexer.cpp
@@ -4,7 +4,6 @@
 #include "Utility/IndexerUtilities.h"
 #include "Sound/SoundCue.h"
 #include "EdGraph/EdGraphNode.h"
-#include "Sound/SoundNode.h"
 #include "SearchSerializer.h"
 
 enum class ESoundCueIndexerVersion
@@ -43,4 +42,4 @@ void FSoundCueIndexer::IndexAsset(const UObject* InAssetObject, FSearchSerialize
 			}
 		}
 	}
-}
\ No newline at end of file
+}
diff --git a/Engine/Plugins/Editor/AssetSearch/Source/Private/Indexers/WidgetBlueprintIndexer.cpp b/Engine/Plugins/Editor/AssetSearch/Source/Private/Indexers/WidgetBlueprintIndexer.cpp
--- a/Engine/Plugins/Editor/AssetSearch/Source/Private/Indexers/WidgetBlueprintIndexer.cpp
+++ b/Engine/Plugins/Editor/AssetSearch/Source/Private/Indexers/WidgetBlueprintIndexer.cpp
@@ -1,11 +1,8 @@
 // Copyright Epic Games, Inc. All Rights Reserved.
 
 #include "WidgetBlueprintIndexer.h"
-#include "Engine/Blueprint.h"
 #include "WidgetBlueprint.h"
-#include "Blueprint/WidgetBlueprintGeneratedClass.h"
 #include "Components/Widget.h"
-#include "Misc/StringBuilder.h"
 #include "Utility/IndexerUtilities.h"
 #include "SearchSerializer.h"
 
@@ -83,4 +80,4 @@ void FWidgetBlueprintIndexer::IndexAsset(const UObject* InAssetObject, FSearchSe
 	}
 }
 
-#undef LOCTEXT_NAMESPACE
\ No newline at end of file
+#undef LOCTEXT_NAMESPACE
diff --git a/Engine/Plugins/Editor/AssetSearch/Source/Private/Providers/AssetRegistrySearchProvider.h b/Engine/Plugins/Editor/AssetSearch/Source/Private/Providers/AssetRegistrySearchProvider.h
--- a/Engine/Plugins/Editor/AssetSearch/Source/Private/Providers/AssetRegistrySearchProvider.h
+++ b/Engine/Plugins/Editor/AssetSearch/Source/Private/Providers/AssetRegistrySearchProvider.h
@@ -2,7 +2,6 @@
 
 #pragma once
 
-#include "CoreMinimal.h"
 #include "ISearchProvider.h"
 
 class FAssetRegistrySearchProvider : public ISearchProvider
@@ -10,4 +9,4 @@ class FAssetRegistrySearchProvider : public ISearchProvider
 public:
 	virtual ~FAssetRegistrySearchProvider() { }
 	virtual void Search(FSearchQueryPtr SearchQuery) override;
-};
\ No newline at end of file
+};
diff --git a/Engine/Plugins/Editor/AssetSearch/Source/Private/SearchSerializer.cpp b/Engine/Plugins/Editor/AssetSearch/Source/Private/SearchSerializer.cpp
--- a/Engine/Plugins/Editor/AssetSearch/Source/Private/SearchSerializer.cpp
+++ b/Engine/Plugins/Editor/AssetSearch/Source/Private/SearchSerializer.cpp
@@ -1,11 +1,7 @@
 // Copyright Epic Games, Inc. All Rights Reserved.
 
 #include "SearchSerializer.h"
-#include "JsonObjectConverter.h"
-#include "Serialization/JsonWriter.h"
-#include "Serialization/JsonSerializer.h"
 #include "Policies/CondensedJsonPrintPolicy.h"
-#include "UObject/UnrealType.h"
 #include "UObject/TextProperty.h"
 #include "IAssetIndexer.h"
 
diff --git a/Engine/Plugins/Editor/AssetSearch/Source/Private/Settings/SearchProjectSettings.h b/Engine/Plugins/Editor/AssetSearch/Source/Private/Settings/SearchProjectSettings.h
--- a/Engine/Plugins/Editor/AssetSearch/Source/Private/Settings/SearchProjectSettings.h
+++ b/Engine/Plugins/Editor/AssetSearch/Source/Private/Settings/SearchProjectSettings.h
@@ -3,9 +3,10 @@
 #pragma once
 
 #include "Engine/DeveloperSettings.h"
-#include "Engine/EngineTypes.h"
 #include "SearchProjectSettings.generated.h"
 
+struct FDirectoryPath;
+
 UCLASS(config = Editor, defaultconfig, meta=(DisplayName="Search"))
 class USearchProjectSettings : public UDeveloperSettings
 {
diff --git a/Engine/Plugins/Editor/AssetSearch/Source/Private/Settings/SearchUserSettings.h b/Engine/Plugins/Editor/AssetSearch/Source/Private/Settings/SearchUserSettings.h
--- a/Engine/Plugins/Editor/AssetSearch/Source/Private/Settings/SearchUserSettings.h
+++ b/Engine/Plugins/Editor/AssetSearch/Source/Private/Settings/SearchUserSettings.h
@@ -3,9 +3,10 @@
 #pragma once
 
 #include "Engine/DeveloperSettings.h"
-#include "Engine/EngineTypes.h"
 #include "SearchUserSettings.generated.h"
 
+struct FDirectoryPath;
+
 USTRUCT()
 struct FSearchPerformance
 {
diff --git a/Engine/Plugins/Editor/AssetSearch/Source/Private/Utility/IndexerUtilities.cpp b/Engine/Plugins/Editor/AssetSearch/Source/Private/Utility/IndexerUtilities.cpp
--- a/Engine/Plugins/Editor/AssetSearch/Source/Private/Utility/IndexerUtilities.cpp
+++ b/Engine/Plugins/Editor/AssetSearch/Source/Private/Utility/IndexerUtilities.cpp
@@ -1,11 +1,9 @@
 // Copyright Epic Games, Inc. All Rights Reserved.
 
 #include "Utility/IndexerUtilities.h"
-#include "UObject/UnrealType.h"
 #include "GameplayTagContainer.h"
-#include "Internationalization/Text.h"
+#include "UObject/EnumProperty.h"
 #include "UObject/TextProperty.h"
-#include "UObject/SoftObjectPtr.h"
 
 static bool IsPropertyIndexable(const TPropertyValueIterator<FProperty>& It, const FProperty* Property)
 {
diff --git a/Engine/Plugins/Editor/AssetSearch/Source/Private/Widgets/SSearchBrowser.cpp b/Engine/Plugins/Editor/AssetSearch/Source/Private/Widgets/SSearchBrowser.cpp
--- a/Engine/Plugins/Editor/AssetSearch/Source/Private/Widgets/SSearchBrowser.cpp
+++ b/Engine/Plugins/Editor/AssetSearch/Source/Private/Widgets/SSearchBrowser.cpp
@@ -1,42 +1,22 @@
 // Copyright Epic Games, Inc. All Rights Reserved.
 
 #include "SSearchBrowser.h"
+#include "Framework/Views/TableViewMetadata.h"
 #include "SSearchTreeRow.h"
 
-#include "ClassViewerModule.h"
-#include "Kismet2/SClassPickerDialog.h"
-#include "Logging/LogMacros.h"
-#include "UObject/ReferenceChainSearch.h"
-#include "UObject/UObjectIterator.h"
-#include "Modules/ModuleManager.h"
 #include "Framework/MultiBox/MultiBoxBuilder.h"
 
 // Common classes for the picker
-#include "Components/Widget.h"
-#include "Blueprint/UserWidget.h"
-#include "GameFramework/Pawn.h"
-#include "GameFramework/PlayerController.h"
-#include "Engine/World.h"
-#include "GameFramework/Actor.h"
-#include "Dialogs/Dialogs.h"
-#include "Widgets/Layout/SBox.h"
-#include "Widgets/Input/SButton.h"
 #include "Widgets/Input/SComboButton.h"
-#include "Widgets/Input/SEditableTextBox.h"
-#include "Widgets/Input/SMultiLineEditableTextBox.h"
 #include "SearchModel.h"
-#include "Styling/AppStyle.h"
 #include "IAssetSearchModule.h"
 #include "AssetRegistry/AssetRegistryModule.h"
 #include "Widgets/Views/STreeView.h"
-#include "Widgets/Layout/SSeparator.h"
 #include "Widgets/Input/SHyperlink.h"
-#include "Widgets/Images/SThrobber.h"
 #include "Widgets/Input/SSearchBox.h"
 #include "Settings/SearchUserSettings.h"
 #include "SearchStyle.h"
 #include "Widgets/SToolTip.h"
-#include "Widgets/Images/SImage.h"
 #include "ThumbnailRendering/ThumbnailManager.h"
 
 #define LOCTEXT_NAMESPACE "SObjectBrowser"
diff --git a/Engine/Plugins/Editor/AssetSearch/Source/Private/Widgets/SSearchBrowser.h b/Engine/Plugins/Editor/AssetSearch/Source/Private/Widgets/SSearchBrowser.h
--- a/Engine/Plugins/Editor/AssetSearch/Source/Private/Widgets/SSearchBrowser.h
+++ b/Engine/Plugins/Editor/AssetSearch/Source/Private/Widgets/SSearchBrowser.h
@@ -2,13 +2,11 @@
 
 #pragma once
 
-#include "Widgets/SCompoundWidget.h"
+#include "SearchQuery.h"
 #include "Widgets/Views/SHeaderRow.h"
-#include "Widgets/Views/SListView.h"
-#include "Widgets/Views/STableViewBase.h"
-#include "SearchModel.h"
-#include "AssetThumbnail.h"
-#include "Widgets/Views/STreeView.h"
+
+class STableViewBase;
+template <typename ItemType> class STreeView;
 
 class IDetailsView;
 class ITableRow;
diff --git a/Engine/Plugins/Editor/AssetSearch/Source/Private/Widgets/SSearchTreeRow.cpp b/Engine/Plugins/Editor/AssetSearch/Source/Private/Widgets/SSearchTreeRow.cpp
--- a/Engine/Plugins/Editor/AssetSearch/Source/Private/Widgets/SSearchTreeRow.cpp
+++ b/Engine/Plugins/Editor/AssetSearch/Source/Private/Widgets/SSearchTreeRow.cpp
@@ -1,20 +1,18 @@
 // Copyright Epic Games, Inc. All Rights Reserved.
 
 #include "SSearchTreeRow.h"
+#include "Engine/Blueprint.h"
 #include "SearchModel.h"
-#include "Widgets/Text/STextBlock.h"
-#include "Widgets/Images/SImage.h"
-#include "Editor.h"
-#include "Subsystems/AssetEditorSubsystem.h"
+#include "IAssetTools.h"
 #include "Styling/SlateIconFinder.h"
-#include "AssetThumbnail.h"
-#include "AssetRegistry/IAssetRegistry.h"
 #include "AssetToolsModule.h"
+#include "IAssetTypeActions.h"
 #include "Kismet2/KismetEditorUtilities.h"
 #include "IMaterialEditor.h"
 
 #include "Materials/Material.h"
 #include "Materials/MaterialFunction.h"
+#include "Misc/Paths.h"
 
 #define LOCTEXT_NAMESPACE "SObjectBrowserTableRow"
 
diff --git a/Engine/Plugins/Editor/AssetSearch/Source/Private/Widgets/SSearchTreeRow.h b/Engine/Plugins/Editor/AssetSearch/Source/Private/Widgets/SSearchTreeRow.h
--- a/Engine/Plugins/Editor/AssetSearch/Source/Private/Widgets/SSearchTreeRow.h
+++ b/Engine/Plugins/Editor/AssetSearch/Source/Private/Widgets/SSearchTreeRow.h
@@ -2,15 +2,9 @@
 
 #pragma once
 
-#include "ClassIconFinder.h"
-#include "CoreMinimal.h"
-#include "Widgets/SWidget.h"
-#include "Widgets/Views/STableRow.h"
 #include "Widgets/Views/SListView.h"
-#include "Widgets/DeclarativeSyntaxSupport.h"
-#include "Input/Reply.h"
-#include "Widgets/Views/STableViewBase.h"
-#include "AssetThumbnail.h"
+
+class FAssetThumbnailPool;
 
 #define LOCTEXT_NAMESPACE "SObjectBrowserTableRow"
 
diff --git a/Engine/Plugins/Editor/AssetSearch/Source/Private/Widgets/SearchModel.h b/Engine/Plugins/Editor/AssetSearch/Source/Private/Widgets/SearchModel.h
--- a/Engine/Plugins/Editor/AssetSearch/Source/Private/Widgets/SearchModel.h
+++ b/Engine/Plugins/Editor/AssetSearch/Source/Private/Widgets/SearchModel.h
@@ -2,14 +2,7 @@
 
 #pragma once
 
-#include "CoreMinimal.h"
-#include "UObject/ObjectMacros.h"
-#include "UObject/Class.h"
-#include "UObject/WeakObjectPtrTemplates.h"
-#include "UObject/Package.h"
-#include "Engine/World.h"
-#include "Templates/SharedPointer.h"
-#include "IAssetSearchModule.h"
+#include "SearchQuery.h"
 
 enum class ESearchNodeType : uint8
 {
diff --git a/Engine/Plugins/Editor/AssetSearch/Source/Private/Widgets/SearchStyle.cpp b/Engine/Plugins/Editor/AssetSearch/Source/Private/Widgets/SearchStyle.cpp
--- a/Engine/Plugins/Editor/AssetSearch/Source/Private/Widgets/SearchStyle.cpp
+++ b/Engine/Plugins/Editor/AssetSearch/Source/Private/Widgets/SearchStyle.cpp
@@ -2,10 +2,10 @@
 
 #include "SearchStyle.h"
 
-#include "Styling/SlateTypes.h"
-#include "Styling/CoreStyle.h"
-#include "Styling/AppStyle.h"
+#include "Brushes/SlateImageBrush.h"
 #include "Interfaces/IPluginManager.h"
+#include "Misc/Paths.h"
+#include "Styling/SlateStyleRegistry.h"
 
 FPluginStyleSet::FPluginStyleSet(const FName& InPluginName, const FName& InStyleSetName)
 	: FSlateStyleSet(InStyleSetName == NAME_None ? InPluginName : FName(*(InPluginName.ToString() + TEXT(".") + InStyleSetName.ToString())))
diff --git a/Engine/Plugins/Editor/AssetSearch/Source/Private/Widgets/SearchStyle.h b/Engine/Plugins/Editor/AssetSearch/Source/Private/Widgets/SearchStyle.h
--- a/Engine/Plugins/Editor/AssetSearch/Source/Private/Widgets/SearchStyle.h
+++ b/Engine/Plugins/Editor/AssetSearch/Source/Private/Widgets/SearchStyle.h
@@ -3,7 +3,6 @@
 #pragma once
 
 #include "Styling/SlateStyle.h"
-#include "Styling/SlateStyleRegistry.h"
 
 /**
  * 
diff --git a/Engine/Plugins/Editor/BlueprintHeaderView/Source/BlueprintHeaderView/Private/BlueprintHeaderView.cpp b/Engine/Plugins/Editor/BlueprintHeaderView/Source/BlueprintHeaderView/Private/BlueprintHeaderView.cpp
--- a/Engine/Plugins/Editor/BlueprintHeaderView/Source/BlueprintHeaderView/Private/BlueprintHeaderView.cpp
+++ b/Engine/Plugins/Editor/BlueprintHeaderView/Source/BlueprintHeaderView/Private/BlueprintHeaderView.cpp
@@ -1,30 +1,23 @@
 // Copyright Epic Games, Inc. All Rights Reserved.
 
 #include "BlueprintHeaderView.h"
-#include "Widgets/SWidget.h"
-#include "Modules/ModuleManager.h"
-#include "Widgets/DeclarativeSyntaxSupport.h"
-#include "Framework/Application/SlateApplication.h"
-#include "Textures/SlateIcon.h"
-#include "Framework/Docking/TabManager.h"
-#include "Styling/AppStyle.h"
+#include "Brushes/SlateImageBrush.h"
 #include "WorkspaceMenuStructure.h"
+#include "Misc/Paths.h"
 #include "WorkspaceMenuStructureModule.h"
+#include "ToolMenu.h"
 #include "Widgets/Docking/SDockTab.h"
 #include "SBlueprintHeaderView.h"
 #include "ContentBrowserModule.h"
-#include "Framework/Commands/UIAction.h"
 #include "Engine/Blueprint.h"
 #include "Engine/UserDefinedStruct.h"
-#include "Framework/MultiBox/MultiBoxBuilder.h"
+#include "ToolMenuSection.h"
 #include "ToolMenus.h"
 #include "Toolkits/AssetEditorToolkit.h"
 #include "Toolkits/AssetEditorToolkitMenuContext.h"
 #include "BlueprintHeaderViewSettings.h"
 #include "Styling/SlateStyle.h"
 #include "Styling/SlateStyleRegistry.h"
-#include "Styling/SlateBrush.h"
-#include "Styling/SlateStyleMacros.h"
 
 #define LOCTEXT_NAMESPACE "BlueprintHeaderViewApp"
 
diff --git a/Engine/Plugins/Editor/BlueprintHeaderView/Source/BlueprintHeaderView/Private/BlueprintHeaderViewSettings.cpp b/Engine/Plugins/Editor/BlueprintHeaderView/Source/BlueprintHeaderView/Private/BlueprintHeaderViewSettings.cpp
--- a/Engine/Plugins/Editor/BlueprintHeaderView/Source/BlueprintHeaderView/Private/BlueprintHeaderViewSettings.cpp
+++ b/Engine/Plugins/Editor/BlueprintHeaderView/Source/BlueprintHeaderView/Private/BlueprintHeaderViewSettings.cpp
@@ -2,6 +2,7 @@
 
 #include "BlueprintHeaderViewSettings.h"
 #include "BlueprintHeaderView.h"
+#include "Styling/SlateTypes.h"
 #include "Styling/StyleColors.h"
 
 #include UE_INLINE_GENERATED_CPP_BY_NAME(BlueprintHeaderViewSettings)
diff --git a/Engine/Plugins/Editor/BlueprintHeaderView/Source/BlueprintHeaderView/Private/HeaderViewClassListItem.cpp b/Engine/Plugins/Editor/BlueprintHeaderView/Source/BlueprintHeaderView/Private/HeaderViewClassListItem.cpp
--- a/Engine/Plugins/Editor/BlueprintHeaderView/Source/BlueprintHeaderView/Private/HeaderViewClassListItem.cpp
+++ b/Engine/Plugins/Editor/BlueprintHeaderView/Source/BlueprintHeaderView/Private/HeaderViewClassListItem.cpp
@@ -1,7 +1,7 @@
 // Copyright Epic Games, Inc. All Rights Reserved.
 
 #include "HeaderViewClassListItem.h"
-#include "Engine/Blueprint.h"
+#include "IAssetTools.h"
 #include "String/LineEndings.h"
 #include "Kismet2/BlueprintEditorUtils.h"
 #include "Framework/MultiBox/MultiBoxBuilder.h"
diff --git a/Engine/Plugins/Editor/BlueprintHeaderView/Source/BlueprintHeaderView/Private/HeaderViewClassListItem.h b/Engine/Plugins/Editor/BlueprintHeaderView/Source/BlueprintHeaderView/Private/HeaderViewClassListItem.h
--- a/Engine/Plugins/Editor/BlueprintHeaderView/Source/BlueprintHeaderView/Private/HeaderViewClassListItem.h
+++ b/Engine/Plugins/Editor/BlueprintHeaderView/Source/BlueprintHeaderView/Private/HeaderViewClassListItem.h
@@ -2,9 +2,10 @@
 
 #pragma once
 
-#include "CoreMinimal.h"
 #include "SBlueprintHeaderView.h"
 
+namespace ETextCommit { enum Type : int; }
+
 class FMenuBuilder;
 
 /** A header view list item that displays the class declaration */
@@ -31,4 +32,4 @@ protected:
 protected:
 	/** Whether this class name is valid C++ (no spaces, special chars, etc) */
 	bool bIsValidName = true;
-};
\ No newline at end of file
+};
diff --git a/Engine/Plugins/Editor/BlueprintHeaderView/Source/BlueprintHeaderView/Private/HeaderViewFunctionListItem.cpp b/Engine/Plugins/Editor/BlueprintHeaderView/Source/BlueprintHeaderView/Private/HeaderViewFunctionListItem.cpp
--- a/Engine/Plugins/Editor/BlueprintHeaderView/Source/BlueprintHeaderView/Private/HeaderViewFunctionListItem.cpp
+++ b/Engine/Plugins/Editor/BlueprintHeaderView/Source/BlueprintHeaderView/Private/HeaderViewFunctionListItem.cpp
@@ -1,11 +1,9 @@
 // Copyright Epic Games, Inc. All Rights Reserved.
 
 #include "HeaderViewFunctionListItem.h"
-#include "Engine/Blueprint.h"
 #include "String/LineEndings.h"
 #include "Kismet2/BlueprintEditorUtils.h"
 #include "K2Node_FunctionEntry.h"
-#include "Misc/EngineVersion.h"
 #include "Framework/MultiBox/MultiBoxBuilder.h"
 #include "Kismet2/KismetEditorUtilities.h"
 
diff --git a/Engine/Plugins/Editor/BlueprintHeaderView/Source/BlueprintHeaderView/Private/HeaderViewFunctionListItem.h b/Engine/Plugins/Editor/BlueprintHeaderView/Source/BlueprintHeaderView/Private/HeaderViewFunctionListItem.h
--- a/Engine/Plugins/Editor/BlueprintHeaderView/Source/BlueprintHeaderView/Private/HeaderViewFunctionListItem.h
+++ b/Engine/Plugins/Editor/BlueprintHeaderView/Source/BlueprintHeaderView/Private/HeaderViewFunctionListItem.h
@@ -2,9 +2,10 @@
 
 #pragma once
 
-#include "CoreMinimal.h"
 #include "SBlueprintHeaderView.h"
 
+namespace ETextCommit { enum Type : int; }
+
 class UK2Node_FunctionEntry;
 
 /** A header view list item that displays a function declaration */
@@ -42,4 +43,4 @@ protected:
 
 	/** Names of any function parameters that are not legal C++ */
 	TArray<FName> IllegalParameters;
-};
\ No newline at end of file
+};
diff --git a/Engine/Plugins/Editor/BlueprintHeaderView/Source/BlueprintHeaderView/Private/HeaderViewStructListItem.cpp b/Engine/Plugins/Editor/BlueprintHeaderView/Source/BlueprintHeaderView/Private/HeaderViewStructListItem.cpp
--- a/Engine/Plugins/Editor/BlueprintHeaderView/Source/BlueprintHeaderView/Private/HeaderViewStructListItem.cpp
+++ b/Engine/Plugins/Editor/BlueprintHeaderView/Source/BlueprintHeaderView/Private/HeaderViewStructListItem.cpp
@@ -2,6 +2,7 @@
 
 #include "HeaderViewStructListItem.h"
 #include "Engine/UserDefinedStruct.h"
+#include "IAssetTools.h"
 #include "String/LineEndings.h"
 #include "Kismet2/StructureEditorUtils.h"
 #include "Framework/MultiBox/MultiBoxBuilder.h"
diff --git a/Engine/Plugins/Editor/BlueprintHeaderView/Source/BlueprintHeaderView/Private/HeaderViewVariableListItem.cpp b/Engine/Plugins/Editor/BlueprintHeaderView/Source/BlueprintHeaderView/Private/HeaderViewVariableListItem.cpp
--- a/Engine/Plugins/Editor/BlueprintHeaderView/Source/BlueprintHeaderView/Private/HeaderViewVariableListItem.cpp
+++ b/Engine/Plugins/Editor/BlueprintHeaderView/Source/BlueprintHeaderView/Private/HeaderViewVariableListItem.cpp
@@ -2,8 +2,6 @@
 
 #include "HeaderViewVariableListItem.h"
 #include "EdMode.h"
-#include "EdGraphSchema_K2.h"
-#include "Engine/Blueprint.h"
 #include "Engine/UserDefinedStruct.h"
 #include "String/LineEndings.h"
 #include "Kismet2/BlueprintEditorUtils.h"
diff --git a/Engine/Plugins/Editor/BlueprintHeaderView/Source/BlueprintHeaderView/Private/SBlueprintHeaderView.cpp b/Engine/Plugins/Editor/BlueprintHeaderView/Source/BlueprintHeaderView/Private/SBlueprintHeaderView.cpp
--- a/Engine/Plugins/Editor/BlueprintHeaderView/Source/BlueprintHeaderView/Private/SBlueprintHeaderView.cpp
+++ b/Engine/Plugins/Editor/BlueprintHeaderView/Source/BlueprintHeaderView/Private/SBlueprintHeaderView.cpp
@@ -2,16 +2,19 @@
 
 #include "SBlueprintHeaderView.h"
 #include "BlueprintHeaderView.h"
+#include "DetailsViewArgs.h"
 #include "Widgets/Layout/SSpacer.h"
-#include "Widgets/Layout/SBox.h"
 //#include "Widgets/Layout/SScrollBox.h"
+#include "EdGraph/EdGraph.h"
 #include "Widgets/Input/SComboButton.h"
+#include "Framework/Commands/UICommandList.h"
 #include "Widgets/Text/SRichTextBlock.h"
+#include "Framework/Views/TableViewMetadata.h"
 #include "Widgets/Input/SButton.h"
+#include "IDetailsView.h"
 #include "PropertyEditorModule.h"
 #include "ContentBrowserModule.h"
 #include "IContentBrowserSingleton.h"
-#include "Engine/Blueprint.h"
 #include "Engine/UserDefinedStruct.h"
 #include "Framework/Text/SlateTextRun.h"
 #include "Framework/Commands/GenericCommands.h"
@@ -22,13 +25,13 @@
 #include "HeaderViewVariableListItem.h"
 #include "EdGraphSchema_K2.h"
 #include "K2Node_FunctionEntry.h"
-#include "Styling/AppStyle.h"
 #include "String/LineEndings.h"
 #include "HAL/PlatformApplicationMisc.h"
 #include "BlueprintHeaderViewSettings.h"
 #include "Internationalization/Regex.h"
 #include "Subsystems/AssetEditorSubsystem.h"
 #include "Editor/EditorEngine.h"
+#include "Widgets/Views/SListView.h"
 
 #define LOCTEXT_NAMESPACE "SBlueprintHeaderView"
 
diff --git a/Engine/Plugins/Editor/BlueprintHeaderView/Source/BlueprintHeaderView/Private/SBlueprintHeaderView.h b/Engine/Plugins/Editor/BlueprintHeaderView/Source/BlueprintHeaderView/Private/SBlueprintHeaderView.h
--- a/Engine/Plugins/Editor/BlueprintHeaderView/Source/BlueprintHeaderView/Private/SBlueprintHeaderView.h
+++ b/Engine/Plugins/Editor/BlueprintHeaderView/Source/BlueprintHeaderView/Private/SBlueprintHeaderView.h
@@ -2,13 +2,12 @@
 
 #pragma once
 
-#include "CoreMinimal.h"
-#include "Widgets/DeclarativeSyntaxSupport.h"
-#include "Widgets/SWidget.h"
 #include "Widgets/SCompoundWidget.h"
-#include "Widgets/Views/SListView.h"
 #include "Misc/NotifyHook.h"
-#include "Kismet2/Kismet2NameValidators.h"
+
+class FMenuBuilder;
+enum class EValidatorResult;
+template <typename ItemType> class SListView;
 
 class SComboButton;
 class UBlueprint;
@@ -176,4 +175,4 @@ private:
 
 	/** List Items source */
 	TArray<FHeaderViewListItemPtr> ListItems;
-};
\ No newline at end of file
+};
diff --git a/Engine/Plugins/Editor/BlueprintMaterialTextureNodes/Source/BlueprintMaterialTextureNodes/Private/BlueprintMaterialTextureNodes.cpp b/Engine/Plugins/Editor/BlueprintMaterialTextureNodes/Source/BlueprintMaterialTextureNodes/Private/BlueprintMaterialTextureNodes.cpp
--- a/Engine/Plugins/Editor/BlueprintMaterialTextureNodes/Source/BlueprintMaterialTextureNodes/Private/BlueprintMaterialTextureNodes.cpp
+++ b/Engine/Plugins/Editor/BlueprintMaterialTextureNodes/Source/BlueprintMaterialTextureNodes/Private/BlueprintMaterialTextureNodes.cpp
@@ -1,6 +1,7 @@
 // Copyright Epic Games, Inc. All Rights Reserved.
 
 #include "BlueprintMaterialTextureNodes.h"
+#include "Modules/ModuleManager.h"
 
 #define LOCTEXT_NAMESPACE "FBlueprintMaterialTextureNodesModule"
 
@@ -19,4 +20,4 @@ void FBlueprintMaterialTextureNodesModule::ShutdownModule()
 
 #undef LOCTEXT_NAMESPACE
 	
-IMPLEMENT_MODULE(FBlueprintMaterialTextureNodesModule, BlueprintMaterialTextureNodes)
\ No newline at end of file
+IMPLEMENT_MODULE(FBlueprintMaterialTextureNodesModule, BlueprintMaterialTextureNodes)
diff --git a/Engine/Plugins/Editor/BlueprintMaterialTextureNodes/Source/BlueprintMaterialTextureNodes/Private/BlueprintMaterialTextureNodesBPLibrary.cpp b/Engine/Plugins/Editor/BlueprintMaterialTextureNodes/Source/BlueprintMaterialTextureNodes/Private/BlueprintMaterialTextureNodesBPLibrary.cpp
--- a/Engine/Plugins/Editor/BlueprintMaterialTextureNodes/Source/BlueprintMaterialTextureNodes/Private/BlueprintMaterialTextureNodesBPLibrary.cpp
+++ b/Engine/Plugins/Editor/BlueprintMaterialTextureNodes/Source/BlueprintMaterialTextureNodes/Private/BlueprintMaterialTextureNodesBPLibrary.cpp
@@ -1,32 +1,26 @@
 // Copyright Epic Games, Inc. All Rights Reserved.
 
 #include "BlueprintMaterialTextureNodesBPLibrary.h"
-#include "CoreMinimal.h"
-#include "BlueprintMaterialTextureNodes.h"
 #include "Engine/Texture2D.h"
 
 //RHI gives access to MaxShaderPlatform and FeatureLevel (i.e. GMaxRHIShaderPlatform)
-#include "RHI.h"
 
 //includes for asset creation
-#include "AssetRegistry/AssetRegistryModule.h"
 #include "AssetToolsModule.h"
 #include "IAssetTools.h"
-#include "IContentBrowserSingleton.h"
+#include "MaterialInstanceBasePropertyOverrides.h"
 #include "PackageTools.h"
 #include "Editor.h"
 #include "Factories/MaterialInstanceConstantFactoryNew.h"
 #include "Logging/MessageLog.h"
-#include "Factories/TextureFactory.h"
-#include "Factories/Factory.h"
 
 //Material and texture includes
 #include "Engine/TextureRenderTarget2D.h"
-#include "Materials/MaterialInstance.h"
+#include "MaterialInterface.h"
 #include "Materials/MaterialInstanceConstant.h"
 #include "MaterialShared.h"
-#include "ImageCore.h"
 #include "ImageCoreUtils.h"
+#include "Misc/PackageName.h"
 #include "TextureResource.h"
 
 #include UE_INLINE_GENERATED_CPP_BY_NAME(BlueprintMaterialTextureNodesBPLibrary)
diff --git a/Engine/Plugins/Editor/ChangelistReview/Source/Private/ChangelistReviewModule.cpp b/Engine/Plugins/Editor/ChangelistReview/Source/Private/ChangelistReviewModule.cpp
--- a/Engine/Plugins/Editor/ChangelistReview/Source/Private/ChangelistReviewModule.cpp
+++ b/Engine/Plugins/Editor/ChangelistReview/Source/Private/ChangelistReviewModule.cpp
@@ -2,9 +2,10 @@
 
 #include "ChangelistReviewModule.h"
 #include "ISourceControlModule.h"
+#include "ToolMenu.h"
 #include "Widgets/Docking/SDockTab.h"
-#include "ISourceControlProvider.h"
 #include "SSourceControlReview.h"
+#include "ToolMenuSection.h"
 #include "ToolMenus.h"
 
 #define LOCTEXT_NAMESPACE "ChangelistReviewModule"
@@ -83,4 +84,4 @@ bool FChangelistReviewModule::CanShowReviewTab() const
 
 #undef LOCTEXT_NAMESPACE
 	
-IMPLEMENT_MODULE(FChangelistReviewModule, ChangelistReview)
\ No newline at end of file
+IMPLEMENT_MODULE(FChangelistReviewModule, ChangelistReview)
diff --git a/Engine/Plugins/Editor/ChangelistReview/Source/Private/SSourceControlReview.cpp b/Engine/Plugins/Editor/ChangelistReview/Source/Private/SSourceControlReview.cpp
--- a/Engine/Plugins/Editor/ChangelistReview/Source/Private/SSourceControlReview.cpp
+++ b/Engine/Plugins/Editor/ChangelistReview/Source/Private/SSourceControlReview.cpp
@@ -2,21 +2,22 @@
 
 #include "SSourceControlReview.h"
 
+#include "AssetRegistry/AssetData.h"
 #include "HAL/FileManager.h"
 #include "ClassIconFinder.h"
+#include "Framework/Views/TableViewMetadata.h"
 #include "ISourceControlModule.h"
-#include "ISourceControlProvider.h"
+#include "Misc/PackageName.h"
 #include "SourceControlOperations.h"
+#include "Misc/Paths.h"
 #include "SSourceControlReviewEntry.h"
 #include "AssetRegistry/AssetRegistryModule.h"
 #include "Kismet/KismetSystemLibrary.h"
 #include "Styling/StarshipCoreStyle.h"
+#include "Widgets/Input/SEditableTextBox.h"
 #include "Widgets/Input/SMultiLineEditableTextBox.h"
-#include "Widgets/Layout/SScrollBox.h"
-#include "Widgets/Layout/SWidgetSwitcher.h"
 #include "Widgets/Notifications/SProgressBar.h"
 #include "Widgets/Input/SButton.h"
-#include "Widgets/Layout/SBox.h"
 #include "Misc/App.h"
 #include "Misc/MessageDialog.h"
 #include "Styling/StyleColors.h"
@@ -24,6 +25,7 @@
 #include "AssetRegistry/AssetRegistryModule.h"
 #include "Brushes/SlateRoundedBoxBrush.h"
 #include "Internationalization/Regex.h"
+#include "Widgets/Views/SListView.h"
 
 #define LOCTEXT_NAMESPACE "SourceControlReview"
 
diff --git a/Engine/Plugins/Editor/ChangelistReview/Source/Private/SSourceControlReview.h b/Engine/Plugins/Editor/ChangelistReview/Source/Private/SSourceControlReview.h
--- a/Engine/Plugins/Editor/ChangelistReview/Source/Private/SSourceControlReview.h
+++ b/Engine/Plugins/Editor/ChangelistReview/Source/Private/SSourceControlReview.h
@@ -1,14 +1,14 @@
 // Copyright Epic Games, Inc. All Rights Reserved.
 
 #pragma once
-#include "ISourceControlOperation.h"
 #include "ISourceControlProvider.h"
 #include "SourceControlOperations.h"
-#include "Widgets/SCompoundWidget.h"
-#include "Widgets/Text/STextBlock.h"
-#include "Widgets/SBoxPanel.h"
-#include "Widgets/Views/SListView.h"
-#include "Widgets/Input/SEditableTextBox.h"
+#include "Widgets/Views/SHeaderRow.h"
+
+class SEditableTextBox;
+class STableViewBase;
+class STextBlock;
+template <typename ItemType> class SListView;
 
 class SSourceControlReviewEntry;
 class SChangelistEditableTextBox;
diff --git a/Engine/Plugins/Editor/ChangelistReview/Source/Private/SSourceControlReviewEntry.cpp b/Engine/Plugins/Editor/ChangelistReview/Source/Private/SSourceControlReviewEntry.cpp
--- a/Engine/Plugins/Editor/ChangelistReview/Source/Private/SSourceControlReviewEntry.cpp
+++ b/Engine/Plugins/Editor/ChangelistReview/Source/Private/SSourceControlReviewEntry.cpp
@@ -2,21 +2,21 @@
 
 #include "SSourceControlReviewEntry.h"
 
+#include "AssetDefinition.h"
 #include "AssetToolsModule.h"
 #include "ClassIconFinder.h"
+#include "Engine/Blueprint.h"
 #include "SourceControlHelpers.h"
 #include "AssetRegistry/AssetRegistryModule.h"
-#include "Engine/Font.h"
+#include "IAssetTools.h"
 #include "Kismet2/KismetEditorUtilities.h"
-#include "Misc/FileHelper.h"
 #include "Misc/MessageDialog.h"
 #include "Settings/EditorLoadingSavingSettings.h"
 #include "Styling/StarshipCoreStyle.h"
 #include "Styling/StyleColors.h"
-#include "Widgets/Text/STextBlock.h"
-#include "Widgets/Layout/SBox.h"
 #include "Widgets/Input/SButton.h"
 #include "Widgets/Layout/SScrollBox.h"
+#include "Widgets/Views/SListView.h"
 
 #define LOCTEXT_NAMESPACE "SourceControlReviewEntry"
 
diff --git a/Engine/Plugins/Editor/ChangelistReview/Source/Private/SSourceControlReviewEntry.h b/Engine/Plugins/Editor/ChangelistReview/Source/Private/SSourceControlReviewEntry.h
--- a/Engine/Plugins/Editor/ChangelistReview/Source/Private/SSourceControlReviewEntry.h
+++ b/Engine/Plugins/Editor/ChangelistReview/Source/Private/SSourceControlReviewEntry.h
@@ -1,10 +1,12 @@
 // Copyright Epic Games, Inc. All Rights Reserved.
 
 #pragma once
-#include "IAssetTypeActions.h"
 #include "SSourceControlReview.h"
-#include "Engine/Blueprint.h"
-#include "Widgets/DeclarativeSyntaxSupport.h"
+#include "Widgets/Views/ITableRow.h"
+#include "Widgets/Views/STableRow.h"
+
+enum EBlueprintType : int;
+struct FRevisionInfo;
 
 class SImage;
 
diff --git a/Engine/Plugins/Editor/ContentBrowser/ContentBrowserAliasDataSoure/Source/ContentBrowserAliasDataSource/Private/ContentBrowserAliasDataSource.cpp b/Engine/Plugins/Editor/ContentBrowser/ContentBrowserAliasDataSoure/Source/ContentBrowserAliasDataSource/Private/ContentBrowserAliasDataSource.cpp
--- a/Engine/Plugins/Editor/ContentBrowser/ContentBrowserAliasDataSoure/Source/ContentBrowserAliasDataSource/Private/ContentBrowserAliasDataSource.cpp
+++ b/Engine/Plugins/Editor/ContentBrowser/ContentBrowserAliasDataSoure/Source/ContentBrowserAliasDataSource/Private/ContentBrowserAliasDataSource.cpp
@@ -2,18 +2,17 @@
 
 #include "ContentBrowserAliasDataSource.h"
 
+#include "AssetRegistry/AssetRegistryModule.h"
 #include "ContentBrowserAssetDataCore.h"
 #include "ContentBrowserAssetDataSource.h"
 #include "ContentBrowserDataSubsystem.h"
 #include "ContentBrowserItemPath.h"
 
 #include "AssetToolsModule.h"
-#include "ObjectTools.h"
-#include "Misc/NamePermissionList.h"
-#include "Modules/ModuleManager.h"
 #include "CollectionManagerModule.h"
-#include "ICollectionManager.h"
+#include "IAssetTools.h"
 #include "IContentBrowserDataModule.h"
+#include "Misc/PackageName.h"
 
 #include UE_INLINE_GENERATED_CPP_BY_NAME(ContentBrowserAliasDataSource)
 
diff --git a/Engine/Plugins/Editor/ContentBrowser/ContentBrowserAliasDataSoure/Source/ContentBrowserAliasDataSource/Private/ContentBrowserAliasDataSourceModule.cpp b/Engine/Plugins/Editor/ContentBrowser/ContentBrowserAliasDataSoure/Source/ContentBrowserAliasDataSource/Private/ContentBrowserAliasDataSourceModule.cpp
--- a/Engine/Plugins/Editor/ContentBrowser/ContentBrowserAliasDataSoure/Source/ContentBrowserAliasDataSource/Private/ContentBrowserAliasDataSourceModule.cpp
+++ b/Engine/Plugins/Editor/ContentBrowser/ContentBrowserAliasDataSoure/Source/ContentBrowserAliasDataSource/Private/ContentBrowserAliasDataSourceModule.cpp
@@ -1,8 +1,6 @@
 // Copyright Epic Games, Inc. All Rights Reserved.
 
 #include "ContentBrowserAliasDataSourceModule.h"
-#include "ContentBrowserAliasDataSource.h"
-#include "Misc/CoreDelegates.h"
 
 void FContentBrowserAliasDataSourceModule::StartupModule()
 {
diff --git a/Engine/Plugins/Editor/ContentBrowser/ContentBrowserAssetDataSource/Source/ContentBrowserAssetDataSource/Private/AssetFileContextMenu.cpp b/Engine/Plugins/Editor/ContentBrowser/ContentBrowserAssetDataSource/Source/ContentBrowserAssetDataSource/Private/AssetFileContextMenu.cpp
--- a/Engine/Plugins/Editor/ContentBrowser/ContentBrowserAssetDataSource/Source/ContentBrowserAssetDataSource/Private/AssetFileContextMenu.cpp
+++ b/Engine/Plugins/Editor/ContentBrowser/ContentBrowserAssetDataSource/Source/ContentBrowserAssetDataSource/Private/AssetFileContextMenu.cpp
@@ -4,56 +4,40 @@
 
 #include "AssetDefinition.h"
 #include "AssetDefinitionRegistry.h"
-#include "Templates/SubclassOf.h"
-#include "Styling/SlateTypes.h"
-#include "Framework/Commands/UIAction.h"
-#include "Textures/SlateIcon.h"
+#include "AssetRegistry/ARFilter.h"
 #include "Engine/Blueprint.h"
 #include "Engine/UserDefinedStruct.h"
 #include "Engine/UserDefinedEnum.h"
+#include "IAssetTypeActions.h"
 #include "Misc/MessageDialog.h"
 #include "HAL/PlatformApplicationMisc.h"
-#include "HAL/FileManager.h"
-#include "Misc/ScopedSlowTask.h"
+#include "Internationalization/TextLocalizationManager.h"
 #include "UObject/MetaData.h"
+#include "ToolMenu.h"
 #include "UObject/PackageFileSummary.h"
+#include "ToolMenuSection.h"
 #include "UObject/UObjectIterator.h"
-#include "Widgets/SBoxPanel.h"
-#include "Widgets/SWindow.h"
-#include "Framework/Application/SlateApplication.h"
-#include "Widgets/Text/STextBlock.h"
 #include "Widgets/Input/SMultiLineEditableTextBox.h"
-#include "Framework/MultiBox/MultiBoxBuilder.h"
-#include "ToolMenus.h"
-#include "ContentBrowserMenuContexts.h"
 #include "ContentBrowserDataMenuContexts.h"
-#include "ContentBrowserDataSource.h"
 #include "Widgets/Input/SButton.h"
-#include "Styling/AppStyle.h"
 #include "EditorReimportHandler.h"
-#include "Components/ActorComponent.h"
-#include "GameFramework/Actor.h"
 #include "UnrealClient.h"
 #include "Materials/MaterialFunction.h"
 #include "Materials/MaterialFunctionInstance.h"
 #include "Materials/Material.h"
 #include "Settings/EditorExperimentalSettings.h"
 #include "Materials/MaterialInstanceConstant.h"
-#include "FileHelpers.h"
 #include "AssetRegistry/AssetRegistryModule.h"
 #include "IAssetTools.h"
 #include "AssetToolsModule.h"
 #include "AssetViewUtils.h"
 #include "Dialogs/Dialogs.h"
 #include "SMetaDataView.h"
-#include "Algo/Transform.h"
 
 #include "ObjectTools.h"
-#include "PackageTools.h"
 #include "Editor.h"
 
 #include "PropertyEditorModule.h"
-#include "Toolkits/GlobalEditorCommonCommands.h"
 #include "ConsolidateWindow.h"
 #include "ReferencedAssetsUtils.h"
 #include "Internationalization/PackageLocalizationUtil.h"
@@ -69,7 +53,6 @@
 
 #include "Internationalization/Culture.h"
 #include "Internationalization/TextPackageNamespaceUtil.h"
-#include "Framework/Commands/GenericCommands.h"
 #include "Framework/Notifications/NotificationManager.h"
 #include "Widgets/Notifications/SNotificationList.h"
 #include "Engine/LevelStreaming.h"
diff --git a/Engine/Plugins/Editor/ContentBrowser/ContentBrowserAssetDataSource/Source/ContentBrowserAssetDataSource/Private/AssetFolderContextMenu.cpp b/Engine/Plugins/Editor/ContentBrowser/ContentBrowserAssetDataSource/Source/ContentBrowserAssetDataSource/Private/AssetFolderContextMenu.cpp
--- a/Engine/Plugins/Editor/ContentBrowser/ContentBrowserAssetDataSource/Source/ContentBrowserAssetDataSource/Private/AssetFolderContextMenu.cpp
+++ b/Engine/Plugins/Editor/ContentBrowser/ContentBrowserAssetDataSource/Source/ContentBrowserAssetDataSource/Private/AssetFolderContextMenu.cpp
@@ -1,18 +1,18 @@
 // Copyright Epic Games, Inc. All Rights Reserved.
 
 #include "AssetFolderContextMenu.h"
+#include "ContentBrowserDataMenuContexts.h"
 #include "Misc/MessageDialog.h"
-#include "Modules/ModuleManager.h"
-#include "UObject/ObjectRedirector.h"
-#include "Textures/SlateIcon.h"
 #include "AssetRegistry/AssetData.h"
 #include "AssetRegistry/ARFilter.h"
 #include "AssetRegistry/AssetRegistryModule.h"
 #include "AssetToolsModule.h"
 #include "AssetViewUtils.h"
-#include "ToolMenus.h"
-#include "ContentBrowserMenuContexts.h"
+#include "IAssetTools.h"
 #include "Misc/ScopedSlowTask.h"
+#include "ToolMenu.h"
+#include "ToolMenuEntry.h"
+#include "ToolMenuSection.h"
 
 #define LOCTEXT_NAMESPACE "ContentBrowser"
 
diff --git a/Engine/Plugins/Editor/ContentBrowser/ContentBrowserAssetDataSource/Source/ContentBrowserAssetDataSource/Private/AssetPropertyTagCache.cpp b/Engine/Plugins/Editor/ContentBrowser/ContentBrowserAssetDataSource/Source/ContentBrowserAssetDataSource/Private/AssetPropertyTagCache.cpp
--- a/Engine/Plugins/Editor/ContentBrowser/ContentBrowserAssetDataSource/Source/ContentBrowserAssetDataSource/Private/AssetPropertyTagCache.cpp
+++ b/Engine/Plugins/Editor/ContentBrowser/ContentBrowserAssetDataSource/Source/ContentBrowserAssetDataSource/Private/AssetPropertyTagCache.cpp
@@ -1,9 +1,6 @@
 // Copyright Epic Games, Inc. All Rights Reserved.
 
 #include "AssetPropertyTagCache.h"
-#include "UObject/Field.h"
-#include "UObject/Class.h"
-#include "UObject/Package.h"
 #include "UObject/UnrealType.h"
 #include "UObject/LinkerLoad.h"
 
diff --git a/Engine/Plugins/Editor/ContentBrowser/ContentBrowserAssetDataSource/Source/ContentBrowserAssetDataSource/Private/ContentBrowserAssetDataCore.cpp b/Engine/Plugins/Editor/ContentBrowser/ContentBrowserAssetDataSource/Source/ContentBrowserAssetDataSource/Private/ContentBrowserAssetDataCore.cpp
--- a/Engine/Plugins/Editor/ContentBrowser/ContentBrowserAssetDataSource/Source/ContentBrowserAssetDataSource/Private/ContentBrowserAssetDataCore.cpp
+++ b/Engine/Plugins/Editor/ContentBrowser/ContentBrowserAssetDataSource/Source/ContentBrowserAssetDataSource/Private/ContentBrowserAssetDataCore.cpp
@@ -3,30 +3,31 @@
 #include "ContentBrowserAssetDataCore.h"
 
 #include "AssetDefinition.h"
+#include "AssetRegistry/ARFilter.h"
 #include "ContentBrowserDataSource.h"
+#include "ContentBrowserAssetDataPayload.h"
 #include "IAssetTools.h"
+#include "ContentBrowserDataSubsystem.h"
 #include "IContentBrowserDataModule.h"
 #include "AssetRegistry/IAssetRegistry.h"
 #include "AssetViewUtils.h"
-#include "AssetTypeActivationOpenedMethod.h"
 #include "AssetPropertyTagCache.h"
+#include "Editor/UnrealEdTypes.h"
 #include "ObjectTools.h"
-#include "Misc/NamePermissionList.h"
-#include "Misc/ScopedSlowTask.h"
 #include "HAL/FileManager.h"
 #include "FileHelpers.h"
 #include "Editor.h"
+#include "IAssetTypeActions.h"
 #include "Subsystems/AssetEditorSubsystem.h"
-#include "ToolMenus.h"
 #include "AssetFolderContextMenu.h"
 #include "AssetFileContextMenu.h"
 #include "ContentBrowserDataUtils.h"
 #include "ContentBrowserDataMenuContexts.h"
-#include "Settings/ContentBrowserSettings.h"
 #include "Engine/Level.h"
-#include "Framework/Notifications/NotificationManager.h"
+#include "Misc/PackageName.h"
 #include "Misc/WarnIfAssetsLoadedInScope.h"
-#include "Widgets/Notifications/SNotificationList.h"
+#include "Misc/Paths.h"
+#include "ToolMenu.h"
 
 #define LOCTEXT_NAMESPACE "ContentBrowserAssetDataSource"
 
diff --git a/Engine/Plugins/Editor/ContentBrowser/ContentBrowserAssetDataSource/Source/ContentBrowserAssetDataSource/Private/ContentBrowserAssetDataPayload.cpp b/Engine/Plugins/Editor/ContentBrowser/ContentBrowserAssetDataSource/Source/ContentBrowserAssetDataSource/Private/ContentBrowserAssetDataPayload.cpp
--- a/Engine/Plugins/Editor/ContentBrowser/ContentBrowserAssetDataSource/Source/ContentBrowserAssetDataSource/Private/ContentBrowserAssetDataPayload.cpp
+++ b/Engine/Plugins/Editor/ContentBrowser/ContentBrowserAssetDataSource/Source/ContentBrowserAssetDataSource/Private/ContentBrowserAssetDataPayload.cpp
@@ -1,13 +1,15 @@
 // Copyright Epic Games, Inc. All Rights Reserved.
 
 #include "ContentBrowserAssetDataPayload.h"
-#include "Modules/ModuleManager.h"
 #include "AssetToolsModule.h"
 #include "AssetThumbnail.h"
 #include "AssetDefinition.h"
 #include "AssetDefinitionRegistry.h"
 #include "Engine/Texture2D.h"
+#include "IAssetTools.h"
 #include "Materials/Material.h"
+#include "IAssetTypeActions.h"
+#include "Misc/PackageName.h"
 
 const FString& FContentBrowserAssetFolderItemDataPayload::GetFilename() const
 {
diff --git a/Engine/Plugins/Editor/ContentBrowser/ContentBrowserAssetDataSource/Source/ContentBrowserAssetDataSource/Private/ContentBrowserAssetDataSource.cpp b/Engine/Plugins/Editor/ContentBrowser/ContentBrowserAssetDataSource/Source/ContentBrowserAssetDataSource/Private/ContentBrowserAssetDataSource.cpp
--- a/Engine/Plugins/Editor/ContentBrowser/ContentBrowserAssetDataSource/Source/ContentBrowserAssetDataSource/Private/ContentBrowserAssetDataSource.cpp
+++ b/Engine/Plugins/Editor/ContentBrowser/ContentBrowserAssetDataSource/Source/ContentBrowserAssetDataSource/Private/ContentBrowserAssetDataSource.cpp
@@ -2,37 +2,37 @@
 
 #include "ContentBrowserAssetDataSource.h"
 #include "ContentBrowserAssetDataCore.h"
-#include "ContentBrowserDataLegacyBridge.h"
-#include "Modules/ModuleManager.h"
 #include "AssetRegistry/AssetRegistryModule.h"
 #include "AssetToolsModule.h"
 #include "CollectionManagerModule.h"
+#include "ContentBrowserAssetDataPayload.h"
 #include "ICollectionManager.h"
 #include "AssetViewUtils.h"
-#include "ObjectTools.h"
+#include "ContentBrowserItemPath.h"
 #include "UObject/GCObjectScopeGuard.h"
 #include "Factories/Factory.h"
-#include "Misc/ScopeExit.h"
-#include "Misc/StringBuilder.h"
-#include "Misc/ScopedSlowTask.h"
 #include "HAL/FileManager.h"
-#include "Interfaces/IPluginManager.h"
-#include "FileHelpers.h"
 #include "Editor.h"
-#include "Subsystems/AssetEditorSubsystem.h"
+#include "IAssetTools.h"
 #include "ToolMenus.h"
+#include "Misc/PackageName.h"
 #include "NewAssetContextMenu.h"
 #include "AssetFolderContextMenu.h"
 #include "AssetFileContextMenu.h"
 #include "ContentBrowserDataSubsystem.h"
 #include "ContentBrowserDataUtils.h"
-#include "IContentBrowserDataModule.h"
 #include "ContentBrowserModule.h"
 #include "IContentBrowserSingleton.h"
-#include "EditorDirectories.h"
 #include "ContentBrowserMenuContexts.h"
+#include "Misc/Paths.h"
 #include "Widgets/Input/SButton.h"
+#include "Subsystems/ImportSubsystem.h"
 #include "Widgets/Images/SImage.h"
+#include "ToolMenu.h"
+#include "ToolMenuEntry.h"
+#include "ToolMenuSection.h"
+#include "Widgets/SBoxPanel.h"
+#include "Widgets/Text/STextBlock.h"
 
 #include UE_INLINE_GENERATED_CPP_BY_NAME(ContentBrowserAssetDataSource)
 
diff --git a/Engine/Plugins/Editor/ContentBrowser/ContentBrowserAssetDataSource/Source/ContentBrowserAssetDataSource/Private/NewAssetContextMenu.cpp b/Engine/Plugins/Editor/ContentBrowser/ContentBrowserAssetDataSource/Source/ContentBrowserAssetDataSource/Private/NewAssetContextMenu.cpp
--- a/Engine/Plugins/Editor/ContentBrowser/ContentBrowserAssetDataSource/Source/ContentBrowserAssetDataSource/Private/NewAssetContextMenu.cpp
+++ b/Engine/Plugins/Editor/ContentBrowser/ContentBrowserAssetDataSource/Source/ContentBrowserAssetDataSource/Private/NewAssetContextMenu.cpp
@@ -1,20 +1,14 @@
 // Copyright Epic Games, Inc. All Rights Reserved.
 
 #include "NewAssetContextMenu.h"
-#include "UObject/UObjectHash.h"
-#include "UObject/UObjectIterator.h"
-#include "Widgets/DeclarativeSyntaxSupport.h"
-#include "Widgets/SCompoundWidget.h"
+#include "ToolMenu.h"
 #include "Widgets/SBoxPanel.h"
+#include "ToolMenuEntry.h"
 #include "Widgets/SOverlay.h"
-#include "Textures/SlateIcon.h"
+#include "ToolMenuSection.h"
 #include "Widgets/Layout/SBorder.h"
 #include "Widgets/Images/SImage.h"
 #include "Widgets/Text/STextBlock.h"
-#include "Widgets/Layout/SBox.h"
-#include "Framework/MultiBox/MultiBoxBuilder.h"
-#include "Styling/AppStyle.h"
-#include "Settings/ContentBrowserSettings.h"
 #include "Factories/Factory.h"
 #include "IAssetTools.h"
 #include "IAssetTypeActions.h"
@@ -23,7 +17,6 @@
 #include "IDocumentation.h"
 #include "ClassIconFinder.h"
 #include "AssetToolsModule.h"
-#include "ToolMenus.h"
 
 #define LOCTEXT_NAMESPACE "ContentBrowser"
 
diff --git a/Engine/Plugins/Editor/ContentBrowser/ContentBrowserClassDataSource/Source/ContentBrowserClassDataSource/Private/ContentBrowserClassDataCore.cpp b/Engine/Plugins/Editor/ContentBrowser/ContentBrowserClassDataSource/Source/ContentBrowserClassDataSource/Private/ContentBrowserClassDataCore.cpp
--- a/Engine/Plugins/Editor/ContentBrowser/ContentBrowserClassDataSource/Source/ContentBrowserClassDataSource/Private/ContentBrowserClassDataCore.cpp
+++ b/Engine/Plugins/Editor/ContentBrowser/ContentBrowserClassDataSource/Source/ContentBrowserClassDataSource/Private/ContentBrowserClassDataCore.cpp
@@ -1,7 +1,9 @@
 // Copyright Epic Games, Inc. All Rights Reserved.
 
 #include "ContentBrowserClassDataCore.h"
+#include "ContentBrowserClassDataPayload.h"
 #include "ContentBrowserDataSource.h"
+#include "Interfaces/IPluginManager.h"
 #include "Misc/PackageName.h"
 #include "AssetViewUtils.h"
 #include "IAssetTypeActions.h"
diff --git a/Engine/Plugins/Editor/ContentBrowser/ContentBrowserClassDataSource/Source/ContentBrowserClassDataSource/Private/ContentBrowserClassDataPayload.cpp b/Engine/Plugins/Editor/ContentBrowser/ContentBrowserClassDataSource/Source/ContentBrowserClassDataSource/Private/ContentBrowserClassDataPayload.cpp
--- a/Engine/Plugins/Editor/ContentBrowser/ContentBrowserClassDataSource/Source/ContentBrowserClassDataSource/Private/ContentBrowserClassDataPayload.cpp
+++ b/Engine/Plugins/Editor/ContentBrowser/ContentBrowserClassDataSource/Source/ContentBrowserClassDataSource/Private/ContentBrowserClassDataPayload.cpp
@@ -2,7 +2,6 @@
 
 #include "ContentBrowserClassDataPayload.h"
 #include "Misc/Paths.h"
-#include "AssetRegistry/AssetData.h"
 #include "AssetThumbnail.h"
 #include "SourceCodeNavigation.h"
 
diff --git a/Engine/Plugins/Editor/ContentBrowser/ContentBrowserClassDataSource/Source/ContentBrowserClassDataSource/Private/ContentBrowserClassDataSource.cpp b/Engine/Plugins/Editor/ContentBrowser/ContentBrowserClassDataSource/Source/ContentBrowserClassDataSource/Private/ContentBrowserClassDataSource.cpp
--- a/Engine/Plugins/Editor/ContentBrowser/ContentBrowserClassDataSource/Source/ContentBrowserClassDataSource/Private/ContentBrowserClassDataSource.cpp
+++ b/Engine/Plugins/Editor/ContentBrowser/ContentBrowserClassDataSource/Source/ContentBrowserClassDataSource/Private/ContentBrowserClassDataSource.cpp
@@ -2,15 +2,13 @@
 
 #include "ContentBrowserClassDataSource.h"
 #include "ContentBrowserClassDataCore.h"
-#include "NativeClassHierarchy.h"
-#include "Modules/ModuleManager.h"
 #include "AssetToolsModule.h"
 #include "CollectionManagerModule.h"
+#include "ContentBrowserClassDataPayload.h"
 #include "ICollectionManager.h"
-#include "Misc/StringBuilder.h"
-#include "Misc/NamePermissionList.h"
-#include "UObject/UObjectHash.h"
+#include "ContentBrowserDataFilter.h"
 #include "ToolMenus.h"
+#include "ContentBrowserDataMenuContexts.h"
 #include "NewClassContextMenu.h"
 #include "GameProjectGenerationModule.h"
 #include "Framework/Docking/TabManager.h"
@@ -18,8 +16,11 @@
 #include "ContentBrowserDataUtils.h"
 #include "ContentBrowserItemPath.h"
 #include "Editor/UnrealEdEngine.h"
+#include "IAssetTools.h"
 #include "Preferences/UnrealEdOptions.h"
+#include "IAssetTypeActions.h"
 #include "UnrealEdGlobals.h"
+#include "ToolMenu.h"
 
 #include UE_INLINE_GENERATED_CPP_BY_NAME(ContentBrowserClassDataSource)
 
diff --git a/Engine/Plugins/Editor/ContentBrowser/ContentBrowserClassDataSource/Source/ContentBrowserClassDataSource/Private/NativeClassHierarchy.cpp b/Engine/Plugins/Editor/ContentBrowser/ContentBrowserClassDataSource/Source/ContentBrowserClassDataSource/Private/NativeClassHierarchy.cpp
--- a/Engine/Plugins/Editor/ContentBrowser/ContentBrowserClassDataSource/Source/ContentBrowserClassDataSource/Private/NativeClassHierarchy.cpp
+++ b/Engine/Plugins/Editor/ContentBrowser/ContentBrowserClassDataSource/Source/ContentBrowserClassDataSource/Private/NativeClassHierarchy.cpp
@@ -2,14 +2,11 @@
 
 #include "NativeClassHierarchy.h"
 #include "Modules/ModuleManager.h"
-#include "UObject/Class.h"
-#include "UObject/UObjectHash.h"
 #include "UObject/UObjectIterator.h"
 #include "UObject/Package.h"
 #include "Misc/App.h"
 #include "Misc/PackageName.h"
 #include "Kismet2/KismetEditorUtilities.h"
-#include "GameProjectGenerationModule.h"
 #include "SourceCodeNavigation.h"
 #include "Interfaces/IPluginManager.h"
 #include "Interfaces/IProjectManager.h"
@@ -627,4 +624,4 @@ TSet<FName> FNativeClassHierarchy::GetGameModules()
 	}
 
 	return GameModules;
-}
\ No newline at end of file
+}
diff --git a/Engine/Plugins/Editor/ContentBrowser/ContentBrowserClassDataSource/Source/ContentBrowserClassDataSource/Private/NativeClassHierarchy.h b/Engine/Plugins/Editor/ContentBrowser/ContentBrowserClassDataSource/Source/ContentBrowserClassDataSource/Private/NativeClassHierarchy.h
--- a/Engine/Plugins/Editor/ContentBrowser/ContentBrowserClassDataSource/Source/ContentBrowserClassDataSource/Private/NativeClassHierarchy.h
+++ b/Engine/Plugins/Editor/ContentBrowser/ContentBrowserClassDataSource/Source/ContentBrowserClassDataSource/Private/NativeClassHierarchy.h
@@ -2,11 +2,12 @@
 
 #pragma once
 
-#include "CoreMinimal.h"
 #include "Containers/ArrayView.h"
+#include "Delegates/Delegate.h"
 #include "HAL/PlatformTime.h"
-#include "Interfaces/IPluginManager.h"
-#include "UObject/UObjectGlobals.h"
+
+enum class EPluginLoadedFrom;
+enum class EReloadCompleteReason;
 
 /**
  * Type of hierarchy node
diff --git a/Engine/Plugins/Editor/ContentBrowser/ContentBrowserClassDataSource/Source/ContentBrowserClassDataSource/Private/NewClassContextMenu.cpp b/Engine/Plugins/Editor/ContentBrowser/ContentBrowserClassDataSource/Source/ContentBrowserClassDataSource/Private/NewClassContextMenu.cpp
--- a/Engine/Plugins/Editor/ContentBrowser/ContentBrowserClassDataSource/Source/ContentBrowserClassDataSource/Private/NewClassContextMenu.cpp
+++ b/Engine/Plugins/Editor/ContentBrowser/ContentBrowserClassDataSource/Source/ContentBrowserClassDataSource/Private/NewClassContextMenu.cpp
@@ -1,8 +1,8 @@
 // Copyright Epic Games, Inc. All Rights Reserved.
 
 #include "NewClassContextMenu.h"
-#include "Styling/AppStyle.h"
-#include "ToolMenus.h"
+#include "ToolMenu.h"
+#include "ToolMenuSection.h"
 
 #define LOCTEXT_NAMESPACE "ContentBrowserClassDataSource"
 
