SourceControl: Fixed crash reverting a checked out -non loaded- map.

diff --git a/Engine/Plugins/Editor/AssetManagerEditor/Source/AssetManagerEditor/Private/AssetSourceControlContextMenu.cpp b/Engine/Plugins/Editor/AssetManagerEditor/Source/AssetManagerEditor/Private/AssetSourceControlContextMenu.cpp
--- a/Engine/Plugins/Editor/AssetManagerEditor/Source/AssetManagerEditor/Private/AssetSourceControlContextMenu.cpp
+++ b/Engine/Plugins/Editor/AssetManagerEditor/Source/AssetManagerEditor/Private/AssetSourceControlContextMenu.cpp
@@ -773,16 +773,20 @@ void FAssetSourceControlContextMenuState::ExecuteSCCRevert() const
 {
 	TArray<FString> PackageNames;
 	GetSelectedPackageNames(PackageNames);
+
 	bool bReloadAll = false;
-	for (FString PackageName : PackageNames)
+	for (const FString& PackageName : PackageNames)
 	{
-		UPackage* Package = FindPackage(nullptr, *PackageName);
-		if (Package->ContainsMap())
+		if (UPackage* Package = FindPackage(nullptr, *PackageName))
 		{
-			bReloadAll = true;
-			break;
+			if (Package->ContainsMap())
+			{
+				bReloadAll = true;
+				break;
+			}
 		}
 	}
+
 	FSourceControlWindows::PromptForRevert(PackageNames, bReloadAll);
 }
 
