SourceControl: Changelist #22880108 added support to the Submit window with support for exposing the diff functionality yes/no, depending on the SourceControl provider used. Changelist #23674745 did the same for the History window. This changelist does the same for the ContentBrowser.

diff --git a/Engine/Plugins/Editor/AssetManagerEditor/Source/AssetManagerEditor/Private/AssetSourceControlContextMenu.cpp b/Engine/Plugins/Editor/AssetManagerEditor/Source/AssetManagerEditor/Private/AssetSourceControlContextMenu.cpp
--- a/Engine/Plugins/Editor/AssetManagerEditor/Source/AssetManagerEditor/Private/AssetSourceControlContextMenu.cpp
+++ b/Engine/Plugins/Editor/AssetManagerEditor/Source/AssetManagerEditor/Private/AssetSourceControlContextMenu.cpp
@@ -382,6 +382,7 @@ void FAssetSourceControlContextMenuState::FillSourceControlSubMenu(UToolMenu* Me
 	const bool bUsesCheckout = SourceControlProvider.UsesCheckout();
 	const bool bUsesFileRevisions = SourceControlProvider.UsesFileRevisions();
 	const bool bUsesReadOnly = SourceControlProvider.UsesLocalReadOnlyState();
+	const bool bUsesDiffAgainstDepot = SourceControlProvider.AllowsDiffAgainstDepot();
 
 	if (bUsesFileRevisions)
 	{
@@ -503,17 +504,20 @@ void FAssetSourceControlContextMenuState::FillSourceControlSubMenu(UToolMenu* Me
 		FIsAsyncProcessingActive::CreateLambda([this]() { return IsStillScanning(CanExecuteSCCHistory()); })
 	);
 
-	AddAsyncMenuEntry(Section,
-		"SCCDiffAgainstDepot",
-		LOCTEXT("SCCDiffAgainstDepot", "Diff Against Depot"),
-		LOCTEXT("SCCDiffAgainstDepotTooltip", "Look at differences between the local and remote version of the selected assets."),
-		FSlateIcon(FRevisionControlStyleManager::GetStyleSetName(), "RevisionControl.Actions.Diff"),
-		FUIAction(
-			FExecuteAction::CreateSP(this, &FAssetSourceControlContextMenuState::ExecuteSCCDiffAgainstDepot),
-			FCanExecuteAction::CreateLambda([this]() { return IsActionEnabled(CanExecuteSCCDiffAgainstDepot()); })
-		),
-		FIsAsyncProcessingActive::CreateLambda([this]() { return IsStillScanning(CanExecuteSCCDiffAgainstDepot()); })
-	);
+	if (bUsesDiffAgainstDepot)
+	{
+		AddAsyncMenuEntry(Section,
+			"SCCDiffAgainstDepot",
+			LOCTEXT("SCCDiffAgainstDepot", "Diff Against Depot"),
+			LOCTEXT("SCCDiffAgainstDepotTooltip", "Look at differences between the local and remote version of the selected assets."),
+			FSlateIcon(FRevisionControlStyleManager::GetStyleSetName(), "RevisionControl.Actions.Diff"),
+			FUIAction(
+				FExecuteAction::CreateSP(this, &FAssetSourceControlContextMenuState::ExecuteSCCDiffAgainstDepot),
+				FCanExecuteAction::CreateLambda([this]() { return IsActionEnabled(CanExecuteSCCDiffAgainstDepot()); })
+			),
+			FIsAsyncProcessingActive::CreateLambda([this]() { return IsStillScanning(CanExecuteSCCDiffAgainstDepot()); })
+		);
+	}
 
 	AddAsyncMenuEntry(Section,
 		"SCCRevert",
