Allow blocking AssetManager OpenReferenceViewerUI from opening

diff --git a/Engine/Plugins/Editor/AssetManagerEditor/Source/AssetManagerEditor/Private/AssetManagerEditorModule.cpp b/Engine/Plugins/Editor/AssetManagerEditor/Source/AssetManagerEditor/Private/AssetManagerEditorModule.cpp
--- a/Engine/Plugins/Editor/AssetManagerEditor/Source/AssetManagerEditor/Private/AssetManagerEditorModule.cpp
+++ b/Engine/Plugins/Editor/AssetManagerEditor/Source/AssetManagerEditor/Private/AssetManagerEditorModule.cpp
@@ -301,6 +301,7 @@ public:
 	virtual void OpenAssetAuditUI(TArray<FAssetData> SelectedAssets) override;
 	virtual void OpenAssetAuditUI(TArray<FAssetIdentifier> SelectedIdentifiers) override;
 	virtual void OpenAssetAuditUI(TArray<FName> SelectedPackages) override;
+	virtual FCanOpenReferenceViewerUI& OnCanOpenReferenceViewerUI() override;
 	virtual void OpenReferenceViewerUI(const TArray<FAssetIdentifier> SelectedIdentifiers, const FReferenceViewerParams ReferenceViewerParams = FReferenceViewerParams()) override;
 	virtual void OpenReferenceViewerUI(const TArray<FName> SelectedPackages, const FReferenceViewerParams ReferenceViewerParams = FReferenceViewerParams()) override;
 	virtual void OpenSizeMapUI(TArray<FAssetIdentifier> SelectedIdentifiers) override;
@@ -371,6 +372,8 @@ private:
 
 	FAssetSourceControlContextMenu AssetSourceControlContextMenu;
 
+	FCanOpenReferenceViewerUI CanOpenReferenceViewerUIDelegate;
+
 	void CreateAssetContextMenu(FToolMenuSection& InSection);
 	void OnExtendContentBrowserCommands(TSharedRef<FUICommandList> CommandList, FOnContentBrowserGetSelection GetSelectionDelegate);
 	void OnExtendLevelEditorCommands(TSharedRef<FUICommandList> CommandList);
@@ -662,11 +665,29 @@ void FAssetManagerEditorModule::OpenAssetAuditUI(TArray<FName> SelectedPackages)
 	}
 }
 
+IAssetManagerEditorModule::FCanOpenReferenceViewerUI& FAssetManagerEditorModule::OnCanOpenReferenceViewerUI()
+{
+	return CanOpenReferenceViewerUIDelegate;
+}
+
 void FAssetManagerEditorModule::OpenReferenceViewerUI(const TArray<FAssetIdentifier> SelectedIdentifiers, const FReferenceViewerParams ReferenceViewerParams)
 {
-	if (SelectedIdentifiers.Num() > 0)
+	if (!SelectedIdentifiers.IsEmpty())
 	{
-		if (TSharedPtr<SDockTab> NewTab = FGlobalTabmanager::Get()->TryInvokeTab(ReferenceViewerTabName))
+		FText ErrorMessage;
+		if (OnCanOpenReferenceViewerUI().IsBound() && !OnCanOpenReferenceViewerUI().Execute(SelectedIdentifiers, &ErrorMessage))
+		{
+			if (!ErrorMessage.IsEmpty())
+			{
+				FNotificationInfo Info(ErrorMessage);
+				Info.ExpireDuration = 5.0f;
+				if (TSharedPtr<SNotificationItem> InfoItem = FSlateNotificationManager::Get().AddNotification(Info))
+				{
+					InfoItem->SetCompletionState(SNotificationItem::CS_Fail);
+				}
+			}
+		}
+		else if (TSharedPtr<SDockTab> NewTab = FGlobalTabmanager::Get()->TryInvokeTab(ReferenceViewerTabName))
 		{
 			TSharedRef<SReferenceViewer> ReferenceViewer = StaticCastSharedRef<SReferenceViewer>(NewTab->GetContent());
 			ReferenceViewer->SetGraphRootIdentifiers(SelectedIdentifiers, ReferenceViewerParams);
diff --git a/Engine/Plugins/Editor/AssetManagerEditor/Source/AssetManagerEditor/Public/AssetManagerEditorModule.h b/Engine/Plugins/Editor/AssetManagerEditor/Source/AssetManagerEditor/Public/AssetManagerEditorModule.h
--- a/Engine/Plugins/Editor/AssetManagerEditor/Source/AssetManagerEditor/Public/AssetManagerEditorModule.h
+++ b/Engine/Plugins/Editor/AssetManagerEditor/Source/AssetManagerEditor/Public/AssetManagerEditorModule.h
@@ -410,6 +410,10 @@ public:
 	virtual void OpenReferenceViewerUI(const TArray<FAssetIdentifier> SelectedIdentifiers, const FReferenceViewerParams ReferenceViewerParams = FReferenceViewerParams()) = 0;
 	virtual void OpenReferenceViewerUI(const TArray<FName> SelectedPackages, const FReferenceViewerParams ReferenceViewerParams = FReferenceViewerParams()) = 0;
 
+	/** Returns false if should not open reference viewer for items */
+	DECLARE_DELEGATE_RetVal_TwoParams(bool, FCanOpenReferenceViewerUI, const TArray<FAssetIdentifier>&, FText*);
+	virtual FCanOpenReferenceViewerUI& OnCanOpenReferenceViewerUI() = 0;
+
 	/** Spawns size map with selected packages */
 	virtual void OpenSizeMapUI(TArray<FAssetIdentifier> SelectedIdentifiers) = 0;
 	virtual void OpenSizeMapUI(TArray<FName> SelectedPackages) = 0;
