AssetManager always initialized: Fix crash on startup when OnMarkPackageDirty is called before InitializeObjectReferences.
#b None, trivial
#rnx

diff --git a/Engine/Plugins/Editor/AssetManagerEditor/Source/AssetManagerEditor/Private/AssetManagerEditorModule.cpp b/Engine/Plugins/Editor/AssetManagerEditor/Source/AssetManagerEditor/Private/AssetManagerEditorModule.cpp
--- a/Engine/Plugins/Editor/AssetManagerEditor/Source/AssetManagerEditor/Private/AssetManagerEditorModule.cpp
+++ b/Engine/Plugins/Editor/AssetManagerEditor/Source/AssetManagerEditor/Private/AssetManagerEditorModule.cpp
@@ -1036,7 +1036,11 @@ void FAssetManagerEditorModule::OnReloadComplete(EReloadCompleteReason Reason)
 
 void FAssetManagerEditorModule::OnMarkPackageDirty(UPackage* Pkg, bool bWasDirty)
 {
-	check(UAssetManager::IsInitialized());
+	if (!UAssetManager::IsInitialized())
+	{
+		// Can be called before InitializeObjectReferences; do nothing in that case
+		return;
+	}
 	UAssetManager& AssetManager = UAssetManager::Get();
 
 	// Check if this package is managed, if so invalidate
