Don't show selected assets in either the referencer or dependency list in Reference Viewer

diff --git a/Engine/Plugins/Editor/AssetManagerEditor/Source/AssetManagerEditor/Private/ReferenceViewer/EdGraph_ReferenceViewer.cpp b/Engine/Plugins/Editor/AssetManagerEditor/Source/AssetManagerEditor/Private/ReferenceViewer/EdGraph_ReferenceViewer.cpp
--- a/Engine/Plugins/Editor/AssetManagerEditor/Source/AssetManagerEditor/Private/ReferenceViewer/EdGraph_ReferenceViewer.cpp
+++ b/Engine/Plugins/Editor/AssetManagerEditor/Source/AssetManagerEditor/Private/ReferenceViewer/EdGraph_ReferenceViewer.cpp
@@ -443,7 +443,8 @@ void UEdGraph_ReferenceViewer::RecursivelyFilterNodeInfos(const FAssetIdentifier
 	bool PassedAssetTypeFilter = FilterCollection && Settings->GetFiltersEnabled() ? FilterCollection->PassesAllFilters(InNodeInfos[InAssetId]) : true;
 	bool PassedSearchTextFilter = IsAssetPassingSearchTextFilter(InAssetId);
 
-	bool PassedAllFilters = Settings->GetFindPathEnabled()  || (PassedAssetTypeFilter && PassedSearchTextFilter);
+	// Don't apply filters in Find Path Mode. Otherwise, check the type and search filters, and also don't include any assets in the central selection (where InCurrentDepth == 0)
+	bool PassedAllFilters = Settings->GetFindPathEnabled() || (PassedAssetTypeFilter && PassedSearchTextFilter && (InCurrentDepth == 0 || !CurrentGraphRootIdentifiers.Contains(InAssetId)));
 
 	InNodeInfos[InAssetId].ChildProvisionSize = NewProvisionSize > 0 ? NewProvisionSize : (PassedAllFilters ? 1 : 0);
 	InNodeInfos[InAssetId].PassedFilters = PassedAllFilters;
