In the asset auditor, show the staged sizes if present in the asset registry, and also make sure that tags show the data from the selected asset registry rather than the in memory one.

diff --git a/Engine/Plugins/Editor/AssetManagerEditor/Source/AssetManagerEditor/Private/AssetManagerEditorModule.cpp b/Engine/Plugins/Editor/AssetManagerEditor/Source/AssetManagerEditor/Private/AssetManagerEditorModule.cpp
--- a/Engine/Plugins/Editor/AssetManagerEditor/Source/AssetManagerEditor/Private/AssetManagerEditorModule.cpp
+++ b/Engine/Plugins/Editor/AssetManagerEditor/Source/AssetManagerEditor/Private/AssetManagerEditorModule.cpp
@@ -104,6 +104,8 @@ const FName IAssetManagerEditorModule::ManagedDiskSizeName = FName("ManagedDiskS
 const FName IAssetManagerEditorModule::TotalUsageName = FName("TotalUsage");
 const FName IAssetManagerEditorModule::CookRuleName = FName("CookRule");
 const FName IAssetManagerEditorModule::ChunksName = FName("Chunks");
+const FName IAssetManagerEditorModule::StageChunkSizeName = FName("Stage_ChunkSize");
+const FName IAssetManagerEditorModule::StageChunkCompressedSizeName = FName("Stage_ChunkCompressedSize");
 
 const FString FAssetManagerEditorRegistrySource::EditorSourceName = TEXT("Editor");
 const FString FAssetManagerEditorRegistrySource::CustomSourceName = TEXT("Custom");
@@ -1247,7 +1249,7 @@ bool FAssetManagerEditorModule::GetStringValueForCustomColumn(const FAssetData&
 	else
 	{
 		// Get base value of asset tag
-		return AssetData.GetTagValue(ColumnName, OutValue);
+		return CurrentRegistrySource->GetAssetTagByObjectPath(AssetData.GetSoftObjectPath(), ColumnName, OutValue);
 	}
 
 	return false;
@@ -1262,7 +1264,7 @@ bool FAssetManagerEditorModule::GetDisplayTextForCustomColumn(const FAssetData&
 
 	UAssetManager& AssetManager = UAssetManager::Get();
 
-	if (ColumnName == ManagedResourceSizeName || ColumnName == ManagedDiskSizeName || ColumnName == DiskSizeName || ColumnName == TotalUsageName)
+	if (ColumnName == ManagedResourceSizeName || ColumnName == ManagedDiskSizeName || ColumnName == DiskSizeName || ColumnName == TotalUsageName || ColumnName == StageChunkSizeName || ColumnName == StageChunkCompressedSizeName)
 	{
 		// Get integer, convert to string
 		int64 IntegerValue = 0;
@@ -1318,7 +1320,7 @@ bool FAssetManagerEditorModule::GetDisplayTextForCustomColumn(const FAssetData&
 	else
 	{
 		// Get base value of asset tag
-		return AssetData.GetTagValue(ColumnName, OutValue);
+		return CurrentRegistrySource->GetAssetTagByObjectPath(AssetData.GetSoftObjectPath(), ColumnName, OutValue);
 	}
 
 	return false;
@@ -1376,6 +1378,14 @@ bool FAssetManagerEditorModule::GetIntegerValueForCustomColumn(const FAssetData&
 	}
 	else if (ColumnName == DiskSizeName)
 	{
+		// If the asset registry has the staged sizes written back, then use that as an accurate representation of the
+		// actual chunks in the package.
+		if (CurrentRegistrySource->GetAssetTagByObjectPath(AssetData.GetSoftObjectPath(), StageChunkSizeName, OutValue))
+		{
+			return true;
+		}
+
+		// Otherwise, use the DiskSize generated by SavePackage.
 		TOptional<FAssetPackageData> FoundData = CurrentRegistrySource->GetAssetPackageDataCopy(AssetData.PackageName);
 
 		if (FoundData && FoundData->DiskSize >= 0)
@@ -1419,10 +1429,7 @@ bool FAssetManagerEditorModule::GetIntegerValueForCustomColumn(const FAssetData&
 	else
 	{
 		// Get base value of asset tag
-		if (AssetData.GetTagValue(ColumnName, OutValue))
-		{
-			return true;
-		}
+		return CurrentRegistrySource->GetAssetTagByObjectPath(AssetData.GetSoftObjectPath(), ColumnName, OutValue);
 	}
 
 	return false;
diff --git a/Engine/Plugins/Editor/AssetManagerEditor/Source/AssetManagerEditor/Private/SAssetAuditBrowser.cpp b/Engine/Plugins/Editor/AssetManagerEditor/Source/AssetManagerEditor/Private/SAssetAuditBrowser.cpp
--- a/Engine/Plugins/Editor/AssetManagerEditor/Source/AssetManagerEditor/Private/SAssetAuditBrowser.cpp
+++ b/Engine/Plugins/Editor/AssetManagerEditor/Source/AssetManagerEditor/Private/SAssetAuditBrowser.cpp
@@ -283,15 +283,20 @@ void SAssetAuditBrowser::Construct(const FArguments& InArgs)
 	// Hide item disk size since we display a custom version here
 	Config.HiddenColumnNames.Add(ContentBrowserItemAttributes::ItemDiskSize.ToString());
 
+	// These are generally not used so hide them by default.
+	Config.HiddenColumnNames.Add(IAssetManagerEditorModule::ManagedDiskSizeName.ToString());
+	Config.HiddenColumnNames.Add(IAssetManagerEditorModule::TotalUsageName.ToString());
+
 	// Add custom columns
 	Config.CustomColumns.Emplace(FPrimaryAssetId::PrimaryAssetTypeTag, LOCTEXT("AssetType", "Primary Type"), LOCTEXT("AssetTypeTooltip", "Primary Asset Type of this asset, if set"), UObject::FAssetRegistryTag::TT_Alphabetical, FOnGetCustomAssetColumnData::CreateSP(this, &SAssetAuditBrowser::GetStringValueForCustomColumn), FOnGetCustomAssetColumnDisplayText::CreateSP(this, &SAssetAuditBrowser::GetDisplayTextForCustomColumn));
 	Config.CustomColumns.Emplace(FPrimaryAssetId::PrimaryAssetNameTag, LOCTEXT("AssetName", "Primary Name"), LOCTEXT("AssetNameTooltip", "Primary Asset Name of this asset, if set"), UObject::FAssetRegistryTag::TT_Alphabetical, FOnGetCustomAssetColumnData::CreateSP(this, &SAssetAuditBrowser::GetStringValueForCustomColumn), FOnGetCustomAssetColumnDisplayText::CreateSP(this, &SAssetAuditBrowser::GetDisplayTextForCustomColumn));
 	Config.CustomColumns.Emplace(IAssetManagerEditorModule::ManagedDiskSizeName, LOCTEXT("ManagedDiskSize", "Disk Size"), LOCTEXT("ManagedDiskSizeTooltip", "Total disk space used by both this and all managed assets"), UObject::FAssetRegistryTag::TT_Numerical, FOnGetCustomAssetColumnData::CreateSP(this, &SAssetAuditBrowser::GetStringValueForCustomColumn), FOnGetCustomAssetColumnDisplayText::CreateSP(this, &SAssetAuditBrowser::GetDisplayTextForCustomColumn));
-	Config.CustomColumns.Emplace(IAssetManagerEditorModule::DiskSizeName, LOCTEXT("DiskSize", "Exclusive Disk Size"), LOCTEXT("DiskSizeTooltip", "Size of saved file on disk for only this asset"), UObject::FAssetRegistryTag::TT_Numerical, FOnGetCustomAssetColumnData::CreateSP(this, &SAssetAuditBrowser::GetStringValueForCustomColumn), FOnGetCustomAssetColumnDisplayText::CreateSP(this, &SAssetAuditBrowser::GetDisplayTextForCustomColumn));
+	Config.CustomColumns.Emplace(IAssetManagerEditorModule::DiskSizeName, LOCTEXT("DiskSize", "Exclusive Disk Size"), LOCTEXT("DiskSizeTooltip", "Size of saved file(s) on disk for this asset's package. If Asset Registry Writeback is enabled and a platform is selected, this will be the uncompressed size of the iostore chunks for the asset's package."), UObject::FAssetRegistryTag::TT_Numerical, FOnGetCustomAssetColumnData::CreateSP(this, &SAssetAuditBrowser::GetStringValueForCustomColumn), FOnGetCustomAssetColumnDisplayText::CreateSP(this, &SAssetAuditBrowser::GetDisplayTextForCustomColumn));
+	Config.CustomColumns.Emplace(IAssetManagerEditorModule::StageChunkCompressedSizeName, LOCTEXT("StagedCompressedSize", "Staged Compressed Size"), LOCTEXT("StagedCompressedSizeTooltip", "Compressed size of iostore chunks for this asset's package. Only visible after staging."), UObject::FAssetRegistryTag::TT_Numerical, FOnGetCustomAssetColumnData::CreateSP(this, &SAssetAuditBrowser::GetStringValueForCustomColumn), FOnGetCustomAssetColumnDisplayText::CreateSP(this, &SAssetAuditBrowser::GetDisplayTextForCustomColumn));
 	Config.CustomColumns.Emplace(IAssetManagerEditorModule::TotalUsageName, LOCTEXT("TotalUsage", "Total Usage"), LOCTEXT("TotalUsageTooltip", "Weighted count of Primary Assets that use this, higher usage means it's more likely to be in memory at runtime"), UObject::FAssetRegistryTag::TT_Numerical, FOnGetCustomAssetColumnData::CreateSP(this, &SAssetAuditBrowser::GetStringValueForCustomColumn), FOnGetCustomAssetColumnDisplayText::CreateSP(this, &SAssetAuditBrowser::GetDisplayTextForCustomColumn));
-	Config.CustomColumns.Emplace(IAssetManagerEditorModule::CookRuleName, LOCTEXT("CookRule", "Cook Rule"), LOCTEXT("CookRuleTooltip", "Rather this asset will be cooked or not"), UObject::FAssetRegistryTag::TT_Alphabetical, FOnGetCustomAssetColumnData::CreateSP(this, &SAssetAuditBrowser::GetStringValueForCustomColumn), FOnGetCustomAssetColumnDisplayText::CreateSP(this, &SAssetAuditBrowser::GetDisplayTextForCustomColumn));
+	Config.CustomColumns.Emplace(IAssetManagerEditorModule::CookRuleName, LOCTEXT("CookRule", "Cook Rule"), LOCTEXT("CookRuleTooltip", "Whether this asset will be cooked or not"), UObject::FAssetRegistryTag::TT_Alphabetical, FOnGetCustomAssetColumnData::CreateSP(this, &SAssetAuditBrowser::GetStringValueForCustomColumn), FOnGetCustomAssetColumnDisplayText::CreateSP(this, &SAssetAuditBrowser::GetDisplayTextForCustomColumn));
 	Config.CustomColumns.Emplace(IAssetManagerEditorModule::ChunksName, LOCTEXT("Chunks", "Chunks"), LOCTEXT("ChunksTooltip", "List of chunks this will be added to when cooked"), UObject::FAssetRegistryTag::TT_Alphabetical, FOnGetCustomAssetColumnData::CreateSP(this, &SAssetAuditBrowser::GetStringValueForCustomColumn), FOnGetCustomAssetColumnDisplayText::CreateSP(this, &SAssetAuditBrowser::GetDisplayTextForCustomColumn));
-	
+
 	// Ignore these tags as we added them as custom columns
 	AssetRegistryTagsToIgnore.Add(FPrimaryAssetId::PrimaryAssetTypeTag);
 	AssetRegistryTagsToIgnore.Add(FPrimaryAssetId::PrimaryAssetNameTag);
diff --git a/Engine/Plugins/Editor/AssetManagerEditor/Source/AssetManagerEditor/Public/AssetManagerEditorModule.h b/Engine/Plugins/Editor/AssetManagerEditor/Source/AssetManagerEditor/Public/AssetManagerEditorModule.h
--- a/Engine/Plugins/Editor/AssetManagerEditor/Source/AssetManagerEditor/Public/AssetManagerEditorModule.h
+++ b/Engine/Plugins/Editor/AssetManagerEditor/Source/AssetManagerEditor/Public/AssetManagerEditorModule.h
@@ -215,6 +215,25 @@ public:
 		}
 	}
 
+	template <typename ValueType>
+	bool GetAssetTagByObjectPath(const FSoftObjectPath& InObjectPath, FName InTagName, ValueType& OutTagValue) const
+	{
+		if (bIsEditor)
+		{
+			return AssetRegistryPrivate->GetAssetByObjectPath(InObjectPath, true /* bIncludeOnlyOnDiskAssets */).GetTagValue(InTagName, OutTagValue);
+		}
+		else if (RegistryStatePrivate)
+		{
+			const FAssetData* AssetData = RegistryStatePrivate->GetAssetByObjectPath(InObjectPath);
+			if (AssetData)
+			{
+				return AssetData->GetTagValue(InTagName, OutTagValue);
+			}
+			return false;
+		}
+		return false;
+	}
+
 	UE_DEPRECATED(5.1, "Asset path FNames have been deprecated, use FSoftObjectPath instead.")
 	FAssetData GetAssetByObjectPath(FName ObjectPath) const
 	{
@@ -362,11 +381,13 @@ public:
 	static const FName TotalUsageName;
 	static const FName CookRuleName;
 	static const FName ChunksName;
+	static const FName StageChunkSizeName;
+	static const FName StageChunkCompressedSizeName;
 
 	/** Gets the value of a "virtual" column for an asset data, this will query the AssetManager for you and takes current platform into account. Returns true and sets out parameter if found */
 	virtual bool GetStringValueForCustomColumn(const FAssetData& AssetData, FName ColumnName, FString& OutValue) = 0;
 
-	/** Gets a display text value for a vritual column. Returns true and sets out parameter if found */
+	/** Gets a display text value for a virtual column. Returns true and sets out parameter if found */
 	virtual bool GetDisplayTextForCustomColumn(const FAssetData& AssetData, FName ColumnName, FText& OutValue) = 0;
 
 	/** Gets the value of a "virtual" column for an asset data, this will query the AssetManager for you and takes current platform into account. Returns true and sets out parameter if found */
