Make sure we are looking at all root graph elements when kicking off a graph rebuild to avoid showing any root assets in the dependency or referencer sections

diff --git a/Engine/Plugins/Editor/AssetManagerEditor/Source/AssetManagerEditor/Private/ReferenceViewer/EdGraph_ReferenceViewer.cpp b/Engine/Plugins/Editor/AssetManagerEditor/Source/AssetManagerEditor/Private/ReferenceViewer/EdGraph_ReferenceViewer.cpp
--- a/Engine/Plugins/Editor/AssetManagerEditor/Source/AssetManagerEditor/Private/ReferenceViewer/EdGraph_ReferenceViewer.cpp
+++ b/Engine/Plugins/Editor/AssetManagerEditor/Source/AssetManagerEditor/Private/ReferenceViewer/EdGraph_ReferenceViewer.cpp
@@ -201,16 +201,22 @@ UEdGraphNode_Reference* UEdGraph_ReferenceViewer::ConstructNodes(const TArray<FA
 		// Note to add an empty parent to the root so that if the root node again gets found again as a duplicate, that next parent won't be 
 		// identified as the primary root and also it will appear as having multiple parents.
 		TMap<FAssetIdentifier, FReferenceNodeInfo> NewReferenceNodeInfos;
-		FReferenceNodeInfo& RootNodeInfo = NewReferenceNodeInfos.FindOrAdd( GraphRootIdentifiers[0], FReferenceNodeInfo(GraphRootIdentifiers[0], true));
-		RootNodeInfo.Parents.Emplace(FAssetIdentifier(NAME_None));
+		for (const FAssetIdentifier& RootIdentifier : GraphRootIdentifiers)
+		{
+			FReferenceNodeInfo& RootNodeInfo = NewReferenceNodeInfos.FindOrAdd( RootIdentifier, FReferenceNodeInfo(RootIdentifier, true));
+			RootNodeInfo.Parents.Emplace(FAssetIdentifier(NAME_None));
+		}
 		if (!Settings->GetFindPathEnabled())
 		{
 			RecursivelyPopulateNodeInfos(true, GraphRootIdentifiers, NewReferenceNodeInfos, 0, Settings->GetSearchReferencerDepthLimit());
 		}
 
 		TMap<FAssetIdentifier, FReferenceNodeInfo> NewDependencyNodeInfos;
-		FReferenceNodeInfo& DRootNodeInfo = NewDependencyNodeInfos.FindOrAdd( GraphRootIdentifiers[0], FReferenceNodeInfo(GraphRootIdentifiers[0], false));
-		DRootNodeInfo.Parents.Emplace(FAssetIdentifier(NAME_None));
+		for (const FAssetIdentifier& RootIdentifier : GraphRootIdentifiers)
+		{
+			FReferenceNodeInfo& DRootNodeInfo = NewDependencyNodeInfos.FindOrAdd(RootIdentifier, FReferenceNodeInfo(RootIdentifier, false));
+			DRootNodeInfo.Parents.Emplace(FAssetIdentifier(NAME_None));
+		}
 		if (!Settings->GetFindPathEnabled())
 		{
 			RecursivelyPopulateNodeInfos(false, GraphRootIdentifiers, NewDependencyNodeInfos, 0, Settings->GetSearchDependencyDepthLimit());
@@ -363,8 +369,14 @@ UEdGraphNode_Reference* UEdGraph_ReferenceViewer::RefilterGraph()
 		FAssetIdentifier FirstGraphRootIdentifier = CurrentGraphRootIdentifiers[0];
 
 		// Create the root node
-		bool bRootIsDuplicated = (Settings->IsShowDependencies() && DependencyNodeInfos.Contains(FirstGraphRootIdentifier) && DependencyNodeInfos[FirstGraphRootIdentifier].IsADuplicate()) || 
-								  (Settings->IsShowReferencers() && ReferencerNodeInfos.Contains(FirstGraphRootIdentifier) && ReferencerNodeInfos[FirstGraphRootIdentifier].IsADuplicate());
+		bool bRootIsDuplicated = false;
+
+		for ( const FAssetIdentifier& RootID : CurrentGraphRootIdentifiers )
+		{
+			bRootIsDuplicated |= (Settings->IsShowDependencies() && DependencyNodeInfos.Contains(RootID) && DependencyNodeInfos[RootID].IsADuplicate()) ||
+				(Settings->IsShowReferencers() && ReferencerNodeInfos.Contains(RootID) && ReferencerNodeInfos[RootID].IsADuplicate());
+		}	
+
 		const FReferenceNodeInfo& NodeInfo = Settings->IsShowReferencers() ? ReferencerNodeInfos[FirstGraphRootIdentifier] : DependencyNodeInfos[FirstGraphRootIdentifier];
 		RootNode = CreateReferenceNode();
 		RootNode->SetupReferenceNode(CurrentGraphRootOrigin, CurrentGraphRootIdentifiers, NodeInfo.AssetData, /*bInAllowThumbnail = */ !Settings->IsCompactMode(), /*bIsDuplicate*/ bRootIsDuplicated);
@@ -518,7 +530,7 @@ void UEdGraph_ReferenceViewer::GetSortedLinks(const TArray<FAssetIdentifier>& Id
 		{
 			EDependencyPinCategory& Category = OutLinks.FindOrAdd(LinkToAsset.AssetId, EDependencyPinCategory::LinkEndActive);
 			bool bIsHard = IsHard(LinkToAsset.Properties);
-			bool bIsUsedInGame = (LinkToAsset.Category != EDependencyCategory::Package) | ((LinkToAsset.Properties & EDependencyProperty::Game) != EDependencyProperty::None);
+			bool bIsUsedInGame = (LinkToAsset.Category != EDependencyCategory::Package) || ((LinkToAsset.Properties & EDependencyProperty::Game) != EDependencyProperty::None);
 			Category |= EDependencyPinCategory::LinkEndActive;
 			Category |= bIsHard ? EDependencyPinCategory::LinkTypeHard : EDependencyPinCategory::LinkTypeNone;
 			Category |= bIsUsedInGame ? EDependencyPinCategory::LinkTypeUsedInGame : EDependencyPinCategory::LinkTypeNone;
