Fixes regression where selected assets in the reference viewer were also showing up in referencer/dependency lists when Show Duplicates was off.

diff --git a/Engine/Plugins/Editor/AssetManagerEditor/Source/AssetManagerEditor/Private/ReferenceViewer/EdGraph_ReferenceViewer.cpp b/Engine/Plugins/Editor/AssetManagerEditor/Source/AssetManagerEditor/Private/ReferenceViewer/EdGraph_ReferenceViewer.cpp
--- a/Engine/Plugins/Editor/AssetManagerEditor/Source/AssetManagerEditor/Private/ReferenceViewer/EdGraph_ReferenceViewer.cpp
+++ b/Engine/Plugins/Editor/AssetManagerEditor/Source/AssetManagerEditor/Private/ReferenceViewer/EdGraph_ReferenceViewer.cpp
@@ -371,6 +371,11 @@ UEdGraphNode_Reference* UEdGraph_ReferenceViewer::RefilterGraph()
 		// Create the root node
 		bool bRootIsDuplicated = false;
 
+		for (const FAssetIdentifier& RootID : CurrentGraphRootIdentifiers)
+		{
+			bRootIsDuplicated |= (Settings->IsShowDependencies() && DependencyNodeInfos.Contains(RootID) && DependencyNodeInfos[RootID].IsADuplicate()) ||
+				(Settings->IsShowReferencers() && ReferencerNodeInfos.Contains(RootID) && ReferencerNodeInfos[RootID].IsADuplicate());
+		}
 		for ( const FAssetIdentifier& RootID : CurrentGraphRootIdentifiers )
 		{
 			bRootIsDuplicated |= (Settings->IsShowDependencies() && DependencyNodeInfos.Contains(RootID) && DependencyNodeInfos[RootID].IsADuplicate()) ||
