Fix to depth vs. max depth checking in Reference Viewer's recursive functions

diff --git a/Engine/Plugins/Editor/AssetManagerEditor/Source/AssetManagerEditor/Private/ReferenceViewer/EdGraph_ReferenceViewer.cpp b/Engine/Plugins/Editor/AssetManagerEditor/Source/AssetManagerEditor/Private/ReferenceViewer/EdGraph_ReferenceViewer.cpp
--- a/Engine/Plugins/Editor/AssetManagerEditor/Source/AssetManagerEditor/Private/ReferenceViewer/EdGraph_ReferenceViewer.cpp
+++ b/Engine/Plugins/Editor/AssetManagerEditor/Source/AssetManagerEditor/Private/ReferenceViewer/EdGraph_ReferenceViewer.cpp
@@ -800,13 +800,14 @@ const TSharedPtr<FAssetThumbnailPool>& UEdGraph_ReferenceViewer::GetAssetThumbna
 
 bool UEdGraph_ReferenceViewer::ExceedsMaxSearchDepth(int32 Depth, int32 MaxDepth) const
 {
+
+	const bool bIsWithinDepthLimits = (MaxDepth > 0 && Depth < MaxDepth);
 	// the FindPath feature is not depth limited
  	if (Settings->GetFindPathEnabled())
  	{
  		return false;
  	}
-
- 	else if (Settings->IsSearchDepthLimited() && (Depth > MaxDepth || MaxDepth < 1))
+ 	else if (Settings->IsSearchDepthLimited() && !bIsWithinDepthLimits)
  	{
  		return true;
  	}
