Add Copy References/Paths to editor asset menu.

#tests Used locally on different editors

diff --git a/Engine/Plugins/Editor/AssetManagerEditor/Source/AssetManagerEditor/Private/AssetManagerEditorCommands.cpp b/Engine/Plugins/Editor/AssetManagerEditor/Source/AssetManagerEditor/Private/AssetManagerEditorCommands.cpp
--- a/Engine/Plugins/Editor/AssetManagerEditor/Source/AssetManagerEditor/Private/AssetManagerEditorCommands.cpp
+++ b/Engine/Plugins/Editor/AssetManagerEditor/Source/AssetManagerEditor/Private/AssetManagerEditorCommands.cpp
@@ -18,6 +18,8 @@ FAssetManagerEditorCommands::FAssetManagerEditorCommands() : TCommands<FAssetMan
 
 void FAssetManagerEditorCommands::RegisterCommands()
 {
+	UI_COMMAND(CopyAssetReferences, "Copy References", "Copy the selected assets' references to the clipboard", EUserInterfaceActionType::Button, FInputChord());
+	UI_COMMAND(CopyAssetPaths, "Copy Paths", "Copy the selected assets' file paths to the clipboard", EUserInterfaceActionType::Button, FInputChord());
 	UI_COMMAND(ViewReferences, "Reference Viewer...", "Launches the reference viewer showing the selected assets' references", EUserInterfaceActionType::Button, FInputChord(EModifierKey::Shift | EModifierKey::Alt, EKeys::R));
 	UI_COMMAND(ViewSizeMap, "Size Map...", "Displays an interactive map showing the approximate size of this asset and everything it references", EUserInterfaceActionType::Button, FInputChord(EModifierKey::Shift | EModifierKey::Alt, EKeys::M));
 	UI_COMMAND(ViewShaderCookStatistics, "Shader Cook Statistics...", "Show Shader CookStatistics", EUserInterfaceActionType::Button, FInputChord());
diff --git a/Engine/Plugins/Editor/AssetManagerEditor/Source/AssetManagerEditor/Private/AssetManagerEditorCommands.h b/Engine/Plugins/Editor/AssetManagerEditor/Source/AssetManagerEditor/Private/AssetManagerEditorCommands.h
--- a/Engine/Plugins/Editor/AssetManagerEditor/Source/AssetManagerEditor/Private/AssetManagerEditorCommands.h
+++ b/Engine/Plugins/Editor/AssetManagerEditor/Source/AssetManagerEditor/Private/AssetManagerEditorCommands.h
@@ -14,6 +14,12 @@ public:
 	virtual void RegisterCommands() override;
 	// End of TCommands<> interface
 
+	// Copy asset references to the clipboard
+	TSharedPtr<FUICommandInfo> CopyAssetReferences;
+
+	// Copy asset file paths to the clipboard
+	TSharedPtr<FUICommandInfo> CopyAssetPaths;
+
 	// Shows the reference viewer for the selected assets
 	TSharedPtr<FUICommandInfo> ViewReferences;
 
diff --git a/Engine/Plugins/Editor/AssetManagerEditor/Source/AssetManagerEditor/Private/AssetManagerEditorModule.cpp b/Engine/Plugins/Editor/AssetManagerEditor/Source/AssetManagerEditor/Private/AssetManagerEditorModule.cpp
--- a/Engine/Plugins/Editor/AssetManagerEditor/Source/AssetManagerEditor/Private/AssetManagerEditorModule.cpp
+++ b/Engine/Plugins/Editor/AssetManagerEditor/Source/AssetManagerEditor/Private/AssetManagerEditorModule.cpp
@@ -35,6 +35,7 @@
 #include "Interfaces/ITargetPlatform.h"
 #include "Interfaces/ITargetPlatformManagerModule.h"
 #include "IPlatformFileSandboxWrapper.h"
+#include "HAL/PlatformApplicationMisc.h"
 #include "HAL/PlatformFileManager.h"
 #include "Serialization/ArrayReader.h"
 #include "EdGraphUtilities.h"
@@ -331,6 +332,9 @@ private:
 	//If the package name is a path rather than a package, then it will do this for each package in the path.
 	void FindClassDependencies(FName PackagePath, const TArray<FTopLevelAssetPath>& TargetClasses, UE::AssetRegistry::EDependencyQuery RequiredDependencyFlags);
 
+	void CopyPackageReferences(const TArray<FAssetData> SelectedPackages);
+	void CopyPackagePaths(const TArray<FAssetData> SelectedPackages);
+
 	bool GetPackageDependencyChain(FName SourcePackage, FName TargetPackage, TArray<FName>& VisitedPackages, TArray<FName>& OutDependencyChain, UE::AssetRegistry::EDependencyQuery RequiredFlags);
 	void GetPackageDependenciesPerClass(FName SourcePackage, const TArray<FTopLevelAssetPath>& TargetClasses, TArray<FName>& VisitedPackages, TArray<FName>& OutDependentPackages, UE::AssetRegistry::EDependencyQuery RequiredDependencyFlags);
 
@@ -966,11 +970,13 @@ void FAssetManagerEditorModule::ExtendContentBrowserPathSelectionMenu()
 TSharedRef<FExtender> FAssetManagerEditorModule::OnExtendAssetEditor(const TSharedRef<FUICommandList> CommandList, const TArray<UObject*> ContextSensitiveObjects)
 {
 	TArray<FName> PackageNames;
+	TArray<FAssetData> PackageAssets;
 	for (UObject* EditedAsset : ContextSensitiveObjects)
 	{
 		if (IsValid(EditedAsset) && EditedAsset->IsAsset())
 		{
 			PackageNames.AddUnique(EditedAsset->GetOutermost()->GetFName());
+			PackageAssets.AddUnique(FAssetData(EditedAsset));
 		}
 	}
 
@@ -979,6 +985,14 @@ TSharedRef<FExtender> FAssetManagerEditorModule::OnExtendAssetEditor(const TShar
 	if (PackageNames.Num() > 0)
 	{
 		// It's safe to modify the CommandList here because this is run as the editor UI is created and the payloads are safe
+		CommandList->MapAction(
+			FAssetManagerEditorCommands::Get().CopyAssetReferences,
+			FExecuteAction::CreateRaw(this, &FAssetManagerEditorModule::CopyPackageReferences, PackageAssets));
+
+		CommandList->MapAction(
+			FAssetManagerEditorCommands::Get().CopyAssetPaths,
+			FExecuteAction::CreateRaw(this, &FAssetManagerEditorModule::CopyPackagePaths, PackageAssets));
+
 		CommandList->MapAction(
 			FAssetManagerEditorCommands::Get().ViewReferences,
 			FExecuteAction::CreateRaw(this, &FAssetManagerEditorModule::OpenReferenceViewerUI, PackageNames, FReferenceViewerParams()));
@@ -1012,6 +1026,19 @@ void FAssetManagerEditorModule::ExtendAssetEditorMenu()
 			{
 				if (IsValid(EditedAsset) && EditedAsset->IsAsset())
 				{
+					InSection.AddMenuEntry(
+						FAssetManagerEditorCommands::Get().CopyAssetReferences,
+						TAttribute<FText>(), // Use command Label
+						TAttribute<FText>(), // Use command tooltip
+						FSlateIcon(FAppStyle::GetAppStyleSetName(), "GenericCommands.Copy")
+					);
+					InSection.AddMenuEntry(
+						FAssetManagerEditorCommands::Get().CopyAssetPaths,
+						TAttribute<FText>(), // Use command Label
+						TAttribute<FText>(), // Use command tooltip
+						FSlateIcon(FAppStyle::GetAppStyleSetName(), "GenericCommands.Copy")
+					);
+
 					CreateAssetContextMenu(InSection);
 					break;
 				}
@@ -2128,6 +2155,49 @@ void FAssetManagerEditorModule::FindClassDependencies(FName SourcePackageName, c
 	}
 }
 
+void FAssetManagerEditorModule::CopyPackageReferences(const TArray<FAssetData> SelectedPackages)
+{
+	TArray<FAssetData> SortedItems = SelectedPackages;
+	SortedItems.Sort([](const FAssetData& One, const FAssetData Two)
+		{
+			return One.PackagePath.Compare(Two.PackagePath) < 0;
+		});
+
+	FString ClipboardText = FString::JoinBy(SortedItems, LINE_TERMINATOR, [](const FAssetData& Item)
+		{
+			return Item.GetExportTextName();
+		});
+
+	FPlatformApplicationMisc::ClipboardCopy(*ClipboardText);
+}
+
+void FAssetManagerEditorModule::CopyPackagePaths(const TArray<FAssetData> SelectedPackages)
+{
+	TArray<FAssetData> SortedItems = SelectedPackages;
+	SortedItems.Sort([](const FAssetData& One, const FAssetData Two)
+		{
+			return One.PackagePath.Compare(Two.PackagePath) < 0;
+		});
+
+	FString ClipboardText = FString::JoinBy(SortedItems, LINE_TERMINATOR, [](const FAssetData& Item)
+		{
+			const FString ItemFilename = FPackageName::LongPackageNameToFilename(Item.PackageName.ToString(), FPackageName::GetAssetPackageExtension());
+
+			if (FPaths::FileExists(ItemFilename))
+			{
+				return FPaths::ConvertRelativePathToFull(ItemFilename);
+			}
+			else
+			{
+				// Add a message for when a user tries to copy the path to a file that doesn't exist on disk of the form
+				// <ItemName>: No file on disk
+				return FString::Printf(TEXT("%s: No file on disk"), *Item.AssetName.ToString());
+			}
+		});
+
+	FPlatformApplicationMisc::ClipboardCopy(*ClipboardText);
+}
+
 void FAssetManagerEditorModule::WriteProfileFile(const FString& Extension, const FString& FileContents)
 {
 	const FString PathName = *(FPaths::ProfilingDir() + TEXT("AssetAudit/"));
