AssetManagerEditor: Added Insights' TreeView implementation, including "advanced filtering" support and needed common code.

diff --git a/Engine/Plugins/Editor/AssetManagerEditor/Source/AssetManagerEditor/Private/Insights/Common/AsyncWorkUtils.h b/Engine/Plugins/Editor/AssetManagerEditor/Source/AssetManagerEditor/Private/Insights/Common/AsyncWorkUtils.h
new file mode 100644
--- /dev/null
+++ b/Engine/Plugins/Editor/AssetManagerEditor/Source/AssetManagerEditor/Private/Insights/Common/AsyncWorkUtils.h
@@ -0,0 +1,25 @@
+// Copyright Epic Games, Inc. All Rights Reserved.
+
+#pragma once
+
+#include "CoreTypes.h"
+
+namespace UE
+{
+namespace Insights
+{
+
+class IAsyncOperationStatusProvider
+{
+public:
+	virtual bool IsRunning() const = 0;
+
+	virtual double GetAllOperationsDuration() = 0;
+	virtual double GetCurrentOperationDuration() = 0;
+	virtual uint32 GetOperationCount() const = 0;
+
+	virtual FText GetCurrentOperationName() const = 0;
+};
+
+} // namespace Insights
+} // namespace UE
diff --git a/Engine/Plugins/Editor/AssetManagerEditor/Source/AssetManagerEditor/Private/Insights/Common/InsightsStyle.cpp b/Engine/Plugins/Editor/AssetManagerEditor/Source/AssetManagerEditor/Private/Insights/Common/InsightsStyle.cpp
new file mode 100644
--- /dev/null
+++ b/Engine/Plugins/Editor/AssetManagerEditor/Source/AssetManagerEditor/Private/Insights/Common/InsightsStyle.cpp
@@ -0,0 +1,314 @@
+// Copyright Epic Games, Inc. All Rights Reserved.
+
+#include "InsightsStyle.h"
+
+#include "Framework/Application/SlateApplication.h"
+#include "HAL/LowLevelMemTracker.h"
+#include "Styling/AppStyle.h"
+#include "Styling/SlateStyleMacros.h"
+#include "Styling/SlateStyleRegistry.h"
+#include "Styling/StarshipCoreStyle.h"
+#include "Styling/StyleColors.h"
+#include "Styling/ToolBarStyle.h"
+
+namespace UE
+{
+namespace Insights
+{
+
+////////////////////////////////////////////////////////////////////////////////////////////////////
+// FInsightsStyle
+////////////////////////////////////////////////////////////////////////////////////////////////////
+
+TSharedPtr<FInsightsStyle::FStyle> FInsightsStyle::StyleInstance = nullptr;
+
+////////////////////////////////////////////////////////////////////////////////////////////////////
+
+void FInsightsStyle::Initialize()
+{
+	LLM_SCOPE_BYNAME(TEXT("Insights/Style"));
+
+	// The core style must be initialized before the Insights style
+#if WITH_EDITOR
+	check(FStarshipCoreStyle::IsInitialized());
+#else
+	FSlateApplication::InitializeCoreStyle();
+#endif
+
+	if (!StyleInstance.IsValid())
+	{
+		StyleInstance = Create();
+		FSlateStyleRegistry::RegisterSlateStyle(*StyleInstance);
+	}
+}
+
+////////////////////////////////////////////////////////////////////////////////////////////////////
+
+TSharedRef<FInsightsStyle::FStyle> FInsightsStyle::Create()
+{
+	TSharedRef<class FInsightsStyle::FStyle> NewStyle = MakeShareable(new FInsightsStyle::FStyle(FInsightsStyle::GetStyleSetName()));
+	NewStyle->Initialize();
+	return NewStyle;
+}
+
+////////////////////////////////////////////////////////////////////////////////////////////////////
+
+void FInsightsStyle::Shutdown()
+{
+	FSlateStyleRegistry::UnRegisterSlateStyle(*StyleInstance);
+	ensure(StyleInstance.IsUnique());
+	StyleInstance.Reset();
+}
+
+////////////////////////////////////////////////////////////////////////////////////////////////////
+
+const ISlateStyle& FInsightsStyle::Get()
+{
+	return *StyleInstance;
+}
+
+////////////////////////////////////////////////////////////////////////////////////////////////////
+
+FName FInsightsStyle::GetStyleSetName()
+{
+	static FName StyleSetName(TEXT("AssetManager/InsightsStyle"));
+	return StyleSetName;
+}
+
+////////////////////////////////////////////////////////////////////////////////////////////////////
+// FInsightsStyle::FStyle
+////////////////////////////////////////////////////////////////////////////////////////////////////
+
+FInsightsStyle::FStyle::FStyle(const FName& InStyleSetName)
+	: FSlateStyleSet(InStyleSetName)
+{
+}
+
+////////////////////////////////////////////////////////////////////////////////////////////////////
+
+void FInsightsStyle::FStyle::SyncParentStyles()
+{
+	const ISlateStyle* ParentStyle = GetParentStyle();
+
+	NormalText = ParentStyle->GetWidgetStyle<FTextBlockStyle>("NormalText");
+	Button = ParentStyle->GetWidgetStyle<FButtonStyle>("Button");
+
+	SelectorColor = ParentStyle->GetSlateColor("SelectorColor");
+	SelectionColor = ParentStyle->GetSlateColor("SelectionColor");
+	SelectionColor_Inactive = ParentStyle->GetSlateColor("SelectionColor_Inactive");
+	SelectionColor_Pressed = ParentStyle->GetSlateColor("SelectionColor_Pressed");
+}
+
+////////////////////////////////////////////////////////////////////////////////////////////////////
+
+#define EDITOR_IMAGE_BRUSH(RelativePath, ...) IMAGE_BRUSH("../../../Editor/Slate/" RelativePath, __VA_ARGS__)
+#define EDITOR_IMAGE_BRUSH_SVG(RelativePath, ...) IMAGE_BRUSH_SVG("../../../Editor/Slate/" RelativePath, __VA_ARGS__)
+#define EDITOR_BOX_BRUSH(RelativePath, ...) BOX_BRUSH("../../../Editor/Slate/" RelativePath, __VA_ARGS__)
+#define EDITOR_BORDER_BRUSH(RelativePath, ...) BORDER_BRUSH("../../../Editor/Slate/" RelativePath, __VA_ARGS__)
+#define TODO_IMAGE_BRUSH(...) EDITOR_IMAGE_BRUSH_SVG("Starship/Common/StaticMesh", __VA_ARGS__)
+
+void FInsightsStyle::FStyle::Initialize()
+{
+	SetParentStyleName("CoreStyle");
+
+	// Sync styles from the parent style that will be used as templates for styles defined here
+	SyncParentStyles();
+
+	Set("Mono.9", DEFAULT_FONT("Mono", 9));
+	Set("Mono.10", DEFAULT_FONT("Mono", 10));
+
+	SetContentRoot(FPaths::EngineContentDir() / TEXT("Slate/Starship/Insights"));
+	SetCoreContentRoot(FPaths::EngineContentDir() / TEXT("Slate"));
+
+	const FVector2D Icon12x12(12.0f, 12.0f); // for TreeItem icons
+	const FVector2D Icon16x16(16.0f, 16.0f); // for regular icons
+	const FVector2D Icon20x20(20.0f, 20.0f); // for ToolBar icons
+
+	//////////////////////////////////////////////////
+
+	Set("Border.TB", new CORE_BOX_BRUSH("Icons/Profiler/Profiler_Border_TB_16x", FMargin(4.0f / 16.0f)));
+	Set("Border.L", new CORE_BOX_BRUSH("Icons/Profiler/Profiler_Border_L_16x", FMargin(4.0f / 16.0f)));
+	Set("Border.R", new CORE_BOX_BRUSH("Icons/Profiler/Profiler_Border_R_16x", FMargin(4.0f / 16.0f)));
+
+	//////////////////////////////////////////////////
+
+	Set("Icons.TableTreeView", new EDITOR_IMAGE_BRUSH_SVG("Starship/Common/Spreadsheet", Icon16x16));
+	Set("Icons.TableTreeView.ToolBar", new EDITOR_IMAGE_BRUSH_SVG("Starship/Common/Spreadsheet", Icon20x20));
+
+	//Set("Icons.Filter", new CORE_IMAGE_BRUSH_SVG("Starship/Common/filter", Icon16x16));	//-> use FAppStyle "Icons.Filter"
+	Set("Icons.FilterAddGroup", new EDITOR_IMAGE_BRUSH_SVG("Starship/Common/WorldOutliner", Icon16x16));
+	Set("Icons.ClassicFilter", new IMAGE_BRUSH_SVG("Filter", Icon16x16));
+	Set("Icons.ClassicFilterConfig", new IMAGE_BRUSH_SVG("FilterConfig", Icon16x16));
+
+	Set("Icons.SortBy", new CORE_IMAGE_BRUSH("Icons/Profiler/profiler_SortBy_32x", Icon16x16));
+	//Set("Icons.SortUp", new CORE_IMAGE_BRUSH_SVG("Starship/Common/SortUp", Icon16x16));	//-> use FAppStyle "Icons.SortUp"
+	//Set("Icons.SortDown", new CORE_IMAGE_BRUSH_SVG("Starship/Common/SortDown", Icon16x16));	//-> use FAppStyle "Icons.SortDown"
+
+	Set("Icons.ViewColumn", new CORE_IMAGE_BRUSH("Icons/Profiler/profiler_ViewColumn_32x", Icon16x16));
+	Set("Icons.ResetColumn", new CORE_IMAGE_BRUSH("Icons/Profiler/profiler_ResetColumn_32x", Icon16x16));
+
+	Set("Icons.ExpandAll", new CORE_IMAGE_BRUSH("Icons/Profiler/profiler_ExpandAll_32x", Icon16x16));
+	Set("Icons.CollapseAll", new CORE_IMAGE_BRUSH("Icons/Profiler/profiler_CollapseAll_32x", Icon16x16));
+	Set("Icons.ExpandSelection", new CORE_IMAGE_BRUSH("Icons/Profiler/profiler_ExpandSelection_32x", Icon16x16));
+	Set("Icons.CollapseSelection", new CORE_IMAGE_BRUSH("Icons/Profiler/profiler_CollapseSelection_32x", Icon16x16));
+
+	Set("Icons.AddGraphSeries", new CORE_IMAGE_BRUSH_SVG("Starship/Common/plus", Icon16x16));
+	Set("Icons.RemoveGraphSeries", new CORE_IMAGE_BRUSH_SVG("Starship/Common/close", Icon16x16));
+
+	Set("Icons.Delete", new CORE_IMAGE_BRUSH_SVG("Starship/Common/Delete", Icon16x16));
+
+	Set("TreeTable.RowBackground", new EDITOR_IMAGE_BRUSH("Old/White", Icon16x16, FLinearColor(1.0f, 1.0f, 1.0f, 0.25f)));
+	Set("TreeViewBanner.WarningIcon", new CORE_IMAGE_BRUSH_SVG("Starship/Common/alert-circle", Icon20x20, FStyleColors::Warning));
+
+	Set("Icons.Hint.TreeItem", new IMAGE_BRUSH_SVG("InfoTag_12", Icon12x12));
+	Set("Icons.Group.TreeItem", new CORE_IMAGE_BRUSH_SVG("Starship/Common/folder-closed", Icon12x12));
+
+	Set("TreeTable.TooltipBold", FTextBlockStyle(NormalText)
+		.SetFont(DEFAULT_FONT("Bold", 8))
+		.SetColorAndOpacity(FLinearColor(0.5f, 0.5f, 0.5f, 1.0f))
+		.SetShadowOffset(FVector2D(1.0f, 1.0f))
+		.SetShadowColorAndOpacity(FLinearColor(0.0f, 0.0f, 0.0f, 0.8f))
+	);
+
+	Set("TreeTable.Tooltip", FTextBlockStyle(NormalText)
+		.SetFont(DEFAULT_FONT("Regular", 8))
+		.SetColorAndOpacity(FLinearColor::White)
+		.SetShadowOffset(FVector2D(1.0f, 1.0f))
+		.SetShadowColorAndOpacity(FLinearColor(0.0f, 0.0f, 0.0f, 0.8f))
+	);
+
+	//////////////////////////////////////////////////
+
+	// NormalEditableTextBox && SearchBox
+#if 0
+	{
+		const FEditableTextBoxStyle& NormalEditableTextBoxStyle = FAppStyle::GetWidgetStyle<FEditableTextBoxStyle>("NormalEditableTextBox");
+		const FTextBlockStyle TextBlockStyle = FTextBlockStyle()
+			.SetColorAndOpacity(NormalEditableTextBoxStyle.ForegroundColor)
+			.SetHighlightColor(NormalEditableTextBoxStyle.FocusedForegroundColor)
+			.SetFont(NormalEditableTextBoxStyle.TextStyle.Font)
+			.SetFontSize(uint16(NormalEditableTextBoxStyle.TextStyle.Font.Size));
+		const FEditableTextBoxStyle EditableTextBoxStyle = FEditableTextBoxStyle(NormalEditableTextBoxStyle)
+			.SetPadding(FMargin(6.0f, 4.0f, 6.0f, 4.0f))
+			.SetTextStyle(TextBlockStyle);
+		Set("NormalEditableTextBox", EditableTextBoxStyle);
+
+		const FSearchBoxStyle& NormalSearchBoxStyle = FAppStyle::GetWidgetStyle<FSearchBoxStyle>("SearchBox");
+		const FSearchBoxStyle SearchBoxStyle = FSearchBoxStyle(NormalSearchBoxStyle)
+			.SetTextBoxStyle(EditableTextBoxStyle);
+		Set("SearchBox", SearchBoxStyle);
+	}
+#endif
+
+	// PrimaryToolbar
+#if 0
+	{
+		FToolBarStyle PrimaryToolbarStyle = FAppStyle::Get().GetWidgetStyle<FToolBarStyle>("SlimToolBar");
+
+		Set("PrimaryToolbar", PrimaryToolbarStyle);
+
+		Set("PrimaryToolbar.MinUniformToolbarSize", 40.0f);
+		Set("PrimaryToolbar.MaxUniformToolbarSize", 40.0f);
+	}
+#endif
+
+	// SecondaryToolbar
+#if 0
+	{
+		FToolBarStyle SecondaryToolbarStyle = FAppStyle::Get().GetWidgetStyle<FToolBarStyle>("SlimToolBar");
+
+		SecondaryToolbarStyle.SetBackgroundPadding(         FMargin(4.0f, 4.0f));
+		SecondaryToolbarStyle.SetBlockPadding(              FMargin(2.0f, 0.0f));
+		SecondaryToolbarStyle.SetButtonPadding(             FMargin(0.0f, 0.0f));
+		SecondaryToolbarStyle.SetCheckBoxPadding(           FMargin(2.0f, 0.0f));
+		SecondaryToolbarStyle.SetComboButtonPadding(        FMargin(0.0f, 0.0f));
+		SecondaryToolbarStyle.SetIndentedBlockPadding(      FMargin(2.0f, 0.0f));
+		SecondaryToolbarStyle.SetLabelPadding(              FMargin(2.0f, 0.0f));
+		SecondaryToolbarStyle.SetSeparatorPadding(          FMargin(2.0f, -3.0f));
+
+		SecondaryToolbarStyle.ToggleButton.SetPadding(      FMargin(0.0f, 0.0f));
+
+		SecondaryToolbarStyle.ButtonStyle.SetNormalPadding( FMargin(6.0f, 2.0f, 4.0f, 2.0f));
+		SecondaryToolbarStyle.ButtonStyle.SetPressedPadding(FMargin(6.0f, 2.0f, 4.0f, 2.0f));
+
+		//SecondaryToolbarStyle.IconSize.Set(16.0f, 16.0f);
+
+		Set("SecondaryToolbar", SecondaryToolbarStyle);
+
+		Set("SecondaryToolbar.MinUniformToolbarSize", 32.0f);
+		Set("SecondaryToolbar.MaxUniformToolbarSize", 32.0f);
+	}
+#endif
+
+	// SecondaryToolbar2 (used by AutoScroll and NetPacketContentView toolbars)
+#if 0
+	{
+		FToolBarStyle SecondaryToolbarStyle = FAppStyle::Get().GetWidgetStyle<FToolBarStyle>("SlimToolBar");
+
+		SecondaryToolbarStyle.SetBackgroundPadding(         FMargin(4.0f, 2.0f));
+		SecondaryToolbarStyle.SetBlockPadding(              FMargin(2.0f, 2.0f));
+		SecondaryToolbarStyle.SetButtonPadding(             FMargin(0.0f, 2.0f));
+		SecondaryToolbarStyle.SetCheckBoxPadding(           FMargin(2.0f, 2.0f));
+		SecondaryToolbarStyle.SetComboButtonPadding(        FMargin(0.0f, 2.0f));
+		SecondaryToolbarStyle.SetIndentedBlockPadding(      FMargin(2.0f, 2.0f));
+		SecondaryToolbarStyle.SetLabelPadding(              FMargin(2.0f, 2.0f));
+		SecondaryToolbarStyle.SetSeparatorPadding(          FMargin(2.0f, -1.0f));
+
+		SecondaryToolbarStyle.ToggleButton.SetPadding(      FMargin(0.0f, 0.0f));
+
+		SecondaryToolbarStyle.ButtonStyle.SetNormalPadding( FMargin(3.0f, 0.0f, -1.0f, 0.0f));
+		SecondaryToolbarStyle.ButtonStyle.SetPressedPadding(FMargin(3.0f, 0.0f, -1.0f, 0.0f));
+
+		//SecondaryToolbarStyle.IconSize.Set(16.0f, 16.0f);
+
+		Set("SecondaryToolbar2", SecondaryToolbarStyle);
+
+		Set("SecondaryToolbar2.MinUniformToolbarSize", 32.0f);
+		Set("SecondaryToolbar2.MaxUniformToolbarSize", 32.0f);
+	}
+#endif
+
+	// ToggleButton
+#if 0
+	{
+		Set("ToggleButton", FButtonStyle(Button)
+			.SetNormal(FSlateNoResource())
+			.SetHovered(EDITOR_BOX_BRUSH("Common/RoundedSelection_16x", 4.0f / 16.0f, SelectionColor))
+			.SetPressed(EDITOR_BOX_BRUSH("Common/RoundedSelection_16x", 4.0f / 16.0f, SelectionColor_Pressed)));
+	}
+#endif
+
+	// Common.GotoNativeCodeHyperlink
+#if 0
+	{
+		FTextBlockStyle InheritedFromNativeTextStyle = FTextBlockStyle(NormalText)
+			.SetFont(DEFAULT_FONT("Regular", 10));
+
+		Set("Common.InheritedFromNativeTextStyle", InheritedFromNativeTextStyle);
+
+		// Go to native class hyperlink
+		FButtonStyle EditNativeHyperlinkButton = FButtonStyle()
+			.SetNormal(EDITOR_BORDER_BRUSH("Old/HyperlinkDotted", FMargin(0, 0, 0, 3 / 16.0f)))
+			.SetPressed(FSlateNoResource())
+			.SetHovered(EDITOR_BORDER_BRUSH("Old/HyperlinkUnderline", FMargin(0, 0, 0, 3 / 16.0f)));
+		FHyperlinkStyle EditNativeHyperlinkStyle = FHyperlinkStyle()
+			.SetUnderlineStyle(EditNativeHyperlinkButton)
+			.SetTextStyle(InheritedFromNativeTextStyle)
+			.SetPadding(FMargin(0.0f));
+
+		Set("Common.GotoNativeCodeHyperlink", EditNativeHyperlinkStyle);
+	}
+#endif
+
+	//////////////////////////////////////////////////
+}
+
+#undef TODO_IMAGE_BRUSH
+#undef EDITOR_BOX_BRUSH
+#undef EDITOR_IMAGE_BRUSH_SVG
+#undef EDITOR_IMAGE_BRUSH
+
+////////////////////////////////////////////////////////////////////////////////////////////////////
+
+} // namespace Insights
+} // namespace UE
diff --git a/Engine/Plugins/Editor/AssetManagerEditor/Source/AssetManagerEditor/Private/Insights/Common/Log.h b/Engine/Plugins/Editor/AssetManagerEditor/Source/AssetManagerEditor/Private/Insights/Common/Log.h
new file mode 100644
--- /dev/null
+++ b/Engine/Plugins/Editor/AssetManagerEditor/Source/AssetManagerEditor/Private/Insights/Common/Log.h
@@ -0,0 +1,11 @@
+// Copyright Epic Games, Inc. All Rights Reserved.
+
+#pragma once
+
+#include "CoreMinimal.h"
+#include "HAL/LowLevelMemTracker.h"
+#include "Logging/LogMacros.h"
+
+DECLARE_LOG_CATEGORY_EXTERN(LogInsights, Verbose, All);
+
+LLM_DECLARE_TAG(Insights);
diff --git a/Engine/Plugins/Editor/AssetManagerEditor/Source/AssetManagerEditor/Private/Insights/Common/TimeUtils.cpp b/Engine/Plugins/Editor/AssetManagerEditor/Source/AssetManagerEditor/Private/Insights/Common/TimeUtils.cpp
new file mode 100644
--- /dev/null
+++ b/Engine/Plugins/Editor/AssetManagerEditor/Source/AssetManagerEditor/Private/Insights/Common/TimeUtils.cpp
@@ -0,0 +1,870 @@
+// Copyright Epic Games, Inc. All Rights Reserved.
+
+#include "Insights/Common/TimeUtils.h"
+
+#include "Misc/StringBuilder.h"
+#include "Insights/Common/Log.h"
+
+#include <cmath>
+#include <limits>
+
+namespace UE
+{
+namespace Insights
+{
+namespace TimeUtils
+{
+
+////////////////////////////////////////////////////////////////////////////////////////////////////
+
+FString FormatTimeValue(const double Duration, const int32 NumDigits)
+{
+	if (std::isnan(Duration))
+	{
+		return TEXT("NaN");
+	}
+
+	if (NumDigits <= 0)
+	{
+		return FString::Printf(TEXT("%.0f"), Duration);
+	}
+
+#if !PLATFORM_USE_GENERIC_STRING_IMPLEMENTATION
+	FString Str = FString::Printf(TEXT("%.*f"), NumDigits, Duration);
+#else
+	// proper resolution is tracked as UE-79534
+	TCHAR FormatString[32];
+	FCString::Snprintf(FormatString, sizeof(FormatString), TEXT("%%.%df"), NumDigits);
+	FString Str = FString::Printf(FormatString, Duration);
+#endif
+
+	int32 Index = Str.Len() - 1;
+	while (Index > 0 && Str[Index] == '0') --Index;
+	while (Index > 0 && Str[Index] == '.') --Index;
+	++Index;
+	if (Index < Str.Len())
+	{
+		Str.LeftInline(Index, false);
+	}
+
+	return Str;
+}
+
+////////////////////////////////////////////////////////////////////////////////////////////////////
+
+FString FormatTimeAuto(const double InDuration, const int32 NumDigits)
+{
+	//TestTimeAutoFormatting();
+
+	double Duration = InDuration;
+
+	if (std::isnan(Duration))
+	{
+		return TEXT("NaN");
+	}
+
+	if (Duration == 0.0)
+	{
+		return TEXT("0");
+	}
+
+	TStringBuilder<64> StrBuilder;
+
+	if (Duration < 0.0)
+	{
+		Duration = -Duration;
+		StrBuilder.AppendChar(TEXT('-'));
+	}
+
+	if (Duration == DBL_MAX || Duration == std::numeric_limits<double>::infinity())
+	{
+		StrBuilder.AppendChar(TEXT('∞'));
+	}
+	else if (Duration < TimeUtils::Picosecond)
+	{
+		// (0 .. 1ps)
+		return TEXT("~0");
+	}
+	else if (Duration < TimeUtils::Nanosecond)
+	{
+		// [1ps .. 1ns)
+		if (Duration >= 999.5 * TimeUtils::Picosecond)
+		{
+			StrBuilder.Append(TEXT("1 ns"));
+		}
+		else
+		{
+			const double Picoseconds = Duration * 1000000000000.0 + 0.5;
+			const int32 IntPicoseconds = static_cast<int32>(Picoseconds);
+			ensure(IntPicoseconds <= 999);
+			//if (IntPicoseconds > 999)
+			//{
+			//	StrBuilder.Append(TEXT("1 ns"));
+			//}
+			//else
+			{
+				StrBuilder.Appendf(TEXT("%d ps"), IntPicoseconds);
+			}
+		}
+	}
+	else if (Duration < TimeUtils::Microsecond)
+	{
+		// [1ns .. 1µs)
+		if (Duration >= 999.5 * TimeUtils::Nanosecond)
+		{
+			StrBuilder.Append(TEXT("1 µs"));
+		}
+		else
+		{
+			const double Nanoseconds = Duration * 1000000000.0 + 0.5;
+			const int32 IntNanoseconds = static_cast<int32>(Nanoseconds);
+			ensure(IntNanoseconds <= 999);
+			//if (IntNanoseconds > 999)
+			//{
+			//	StrBuilder.Append(TEXT("1 µs"));
+			//}
+			//else
+			{
+				StrBuilder.Appendf(TEXT("%d ns"), IntNanoseconds);
+			}
+		}
+	}
+	else if (Duration < TimeUtils::Milisecond)
+	{
+		// [1µs .. 1ms)
+		const double Microseconds = Duration * 1000000.0;
+		if (Microseconds >= 999.95)
+		{
+			StrBuilder.Append(TEXT("1 ms"));
+		}
+		else
+		{
+			StrBuilder.Append(FormatTimeValue(Microseconds, NumDigits));
+			StrBuilder.Append(TEXT(" µs"));
+		}
+	}
+	else if (Duration < TimeUtils::Second)
+	{
+		// [1ms .. 1s)
+		const double Miliseconds = Duration * 1000.0;
+		if (Miliseconds >= 999.95)
+		{
+			StrBuilder.Append(TEXT("1s"));
+		}
+		else
+		{
+			StrBuilder.Append(FormatTimeValue(Miliseconds, NumDigits));
+			StrBuilder.Append(TEXT(" ms"));
+		}
+	}
+	else if (Duration < TimeUtils::Minute)
+	{
+		// [1s .. 1m)
+		if (Duration >= 59.95)
+		{
+			StrBuilder.Append(TEXT("1m"));
+		}
+		else
+		{
+			StrBuilder.Append(FormatTimeValue(Duration, NumDigits));
+			StrBuilder.AppendChar(TEXT('s'));
+		}
+	}
+	else if (Duration < TimeUtils::Hour)
+	{
+		// [1m .. 1h)
+		const double Minutes = FMath::FloorToDouble(Duration / TimeUtils::Minute);
+		StrBuilder.Appendf(TEXT("%dm"), static_cast<int32>(Minutes));
+		Duration -= Minutes * TimeUtils::Minute;
+		if (NumDigits <= 1)
+		{
+			const double Seconds = FMath::FloorToDouble(Duration / TimeUtils::Second);
+			if (Seconds > 0.5)
+			{
+				StrBuilder.Appendf(TEXT(" %ds"), static_cast<int32>(Seconds));
+			}
+		}
+		else
+		{
+			StrBuilder.AppendChar(TEXT(' '));
+			StrBuilder.Append(FormatTimeValue(Duration, NumDigits - 1));
+			StrBuilder.AppendChar(TEXT('s'));
+		}
+	}
+	else if (Duration < TimeUtils::Day)
+	{
+		// [1h .. 1d)
+		const double Hours = FMath::FloorToDouble(Duration / TimeUtils::Hour);
+		StrBuilder.Appendf(TEXT("%dh"), static_cast<int32>(Hours));
+		Duration -= Hours * TimeUtils::Hour;
+		const double Minutes = FMath::FloorToDouble(Duration / TimeUtils::Minute);
+		if (Minutes > 0.5)
+		{
+			StrBuilder.Appendf(TEXT(" %dm"), static_cast<int32>(Minutes));
+		}
+	}
+	else
+	{
+		// [1d .. ∞)
+		const double Days = FMath::FloorToDouble(Duration / TimeUtils::Day);
+		StrBuilder.Appendf(TEXT("%dd"), static_cast<int32>(Days));
+		Duration -= Days * TimeUtils::Day;
+		const double Hours = FMath::FloorToDouble(Duration / TimeUtils::Hour);
+		if (Hours > 0.5)
+		{
+			StrBuilder.Appendf(TEXT(" %dh"), static_cast<int32>(Hours));
+		}
+	}
+
+	return FString(StrBuilder);
+}
+
+////////////////////////////////////////////////////////////////////////////////////////////////////
+
+FString FormatTimeMs(const double InDuration, const int32 NumDigits, bool bAddTimeUnit)
+{
+	if (std::isnan(InDuration))
+	{
+		return TEXT("NaN");
+	}
+
+	if (FMath::IsNearlyZero(InDuration, TimeUtils::Picosecond))
+	{
+		return TEXT("0");
+	}
+
+	double Duration = InDuration;
+	TStringBuilder<64> StrBuilder;
+
+	if (Duration < 0.0)
+	{
+		Duration = -Duration;
+		StrBuilder.AppendChar(TEXT('-'));
+	}
+
+	if (Duration == DBL_MAX || Duration == std::numeric_limits<double>::infinity())
+	{
+		StrBuilder.AppendChar(TEXT('∞'));
+	}
+	else
+	{
+		if (NumDigits <= 0)
+		{
+			StrBuilder.Appendf(TEXT("%.0f"), Duration * 1000.0);
+		}
+		else
+		{
+#if !PLATFORM_USE_GENERIC_STRING_IMPLEMENTATION
+			StrBuilder.Appendf(TEXT("%.*f"), NumDigits, Duration * 1000.0);
+#else
+			// proper resolution is tracked as UE-79534
+			TCHAR FormatString[32];
+			FCString::Snprintf(FormatString, sizeof(FormatString), TEXT("%%.%df"), NumDigits);
+			StrBuilder.Appendf(FormatString, Duration * 1000.0);
+#endif
+		}
+
+		if (bAddTimeUnit)
+		{
+			StrBuilder.Append(TEXT(" ms"));
+		}
+	}
+
+	return FString(StrBuilder);
+}
+
+////////////////////////////////////////////////////////////////////////////////////////////////////
+
+FString FormatTime(const double InTime, const double Precision)
+{
+	if (std::isnan(InTime))
+	{
+		return TEXT("NaN");
+	}
+
+	if (FMath::IsNearlyZero(InTime, FMath::Max(TimeUtils::Picosecond, Precision / 10.0f)))
+	{
+		return TEXT("0");
+	}
+
+	double Time = InTime;
+	TStringBuilder<64> StrBuilder;
+
+	if (Time < 0.0)
+	{
+		Time = -Time;
+		StrBuilder.AppendChar(TEXT('-'));
+	}
+
+	if (Time == DBL_MAX || Time == std::numeric_limits<double>::infinity())
+	{
+		StrBuilder.AppendChar(TEXT('∞'));
+		return FString(StrBuilder);
+	}
+
+	bool bIsSpaceNeeded = false;
+
+	int32 Days = static_cast<int32>(Time / TimeUtils::Day);
+	if (Days > 0)
+	{
+		StrBuilder.Appendf(TEXT("%dd"), Days);
+		bIsSpaceNeeded = true;
+		Time -= static_cast<double>(Days) * TimeUtils::Day;
+	}
+
+	if (Precision >= TimeUtils::Day)
+	{
+		if (!bIsSpaceNeeded)
+		{
+			return TEXT("~0");
+		}
+		return FString(StrBuilder);
+	}
+
+	int32 Hours = static_cast<int32>(Time / TimeUtils::Hour);
+	if (Hours > 0)
+	{
+		if (bIsSpaceNeeded)
+		{
+			StrBuilder.AppendChar(TEXT(' '));
+		}
+		StrBuilder.Appendf(TEXT("%dh"), Hours);
+		bIsSpaceNeeded = true;
+		Time -= static_cast<double>(Hours) * TimeUtils::Hour;
+	}
+
+	if (Precision >= TimeUtils::Hour)
+	{
+		if (!bIsSpaceNeeded)
+		{
+			return TEXT("~0");
+		}
+		return FString(StrBuilder);
+	}
+
+	int32 Minutes = static_cast<int32>(Time / TimeUtils::Minute);
+	if (Minutes > 0)
+	{
+		if (bIsSpaceNeeded)
+		{
+			StrBuilder.AppendChar(TEXT(' '));
+		}
+		StrBuilder.Appendf(TEXT("%dm"), Minutes);
+		bIsSpaceNeeded = true;
+		Time -= static_cast<double>(Minutes) * TimeUtils::Minute;
+	}
+
+	if (Precision >= TimeUtils::Minute)
+	{
+		if (!bIsSpaceNeeded)
+		{
+			return TEXT("~0");
+		}
+		return FString(StrBuilder);
+	}
+
+	//TestOptimizationIssue();
+
+	////float Log10 = -FMath::LogX(10.0f, static_cast<float>(Precision));
+	//double Log10 = -log10(Precision);
+	//int32 Digits = (Log10 > 0) ? FMath::CeilToInt(Log10) : 0;
+
+	static const double DigitThresholds[] =
+	{
+		TimeUtils::Second,              // 0 digits
+		TimeUtils::Milisecond * 100,    // 1 digit
+		TimeUtils::Milisecond * 10,     // 2 digits
+		TimeUtils::Milisecond,          // 3 digits
+		TimeUtils::Microsecond * 100,   // 4 digits
+		TimeUtils::Microsecond * 10,    // 5 digits
+		TimeUtils::Microsecond,         // 6 digits
+		TimeUtils::Nanosecond * 100,    // 7 digits
+		TimeUtils::Nanosecond * 10,     // 8 digits
+		TimeUtils::Nanosecond,          // 9 digits
+		TimeUtils::Picosecond * 100,    // 10 digits
+		TimeUtils::Picosecond * 10,     // 11 digits
+		TimeUtils::Picosecond,          // 12 digits
+		0                               // 13 digits
+	};
+	int32 Digits = 0;
+	while (Precision < DigitThresholds[Digits])
+	{
+		Digits++;
+	}
+
+	if (Digits == 0)
+	{
+		int32 Seconds = static_cast<int32>(Time / TimeUtils::Second + 0.5);
+		if (Seconds > 0)
+		{
+			if (bIsSpaceNeeded)
+			{
+				StrBuilder.AppendChar(TEXT(' '));
+			}
+			StrBuilder.Appendf(TEXT("%ds"), Seconds);
+		}
+		else if (!bIsSpaceNeeded)
+		{
+			return TEXT("~0");
+		}
+	}
+	//else if (Digits <= 9)
+	//{
+	//	if (bIsSpaceNeeded)
+	//	{
+	//		StrBuilder.AppendChar(TEXT(' '));
+	//	}
+	//	int32 Seconds = static_cast<int32>(Time / TimeUtils::Second);
+	//	Time -= static_cast<double>(Seconds) * TimeUtils::Second;
+	//	int64 SubSeconds = static_cast<int64>(Time * FMath::Pow(10.0f, Digits) + 0.5);
+	//	StrBuilder.Appendf(TEXT("%d.%0*llds"), Seconds, Digits, SubSeconds);
+	//}
+	else
+	{
+		if (bIsSpaceNeeded)
+		{
+			StrBuilder.AppendChar(TEXT(' '));
+		}
+#if !PLATFORM_USE_GENERIC_STRING_IMPLEMENTATION
+		StrBuilder.Appendf(TEXT("%.*fs"), Digits, Time);
+#else
+		// proper resolution is tracked as UE-79534
+		TCHAR FormatString[32];
+		FCString::Snprintf(FormatString, sizeof(FormatString), TEXT("%%.%dfs"), Digits);
+		StrBuilder.Appendf(FormatString, Time);
+#endif
+	}
+
+	return FString(StrBuilder);
+}
+
+////////////////////////////////////////////////////////////////////////////////////////////////////
+
+FString FormatTimeHMS(const double Time, const double Precision)
+{
+	//TODO: DD:HH:MM:SS.mmm.uuu.nnn.ppp
+	return FormatTime(Time, Precision);
+}
+
+////////////////////////////////////////////////////////////////////////////////////////////////////
+
+void SplitTime(const double InTime, FTimeSplit& OutTimeSplit)
+{
+	double Time = InTime;
+
+	OutTimeSplit.bIsZero = false;
+	OutTimeSplit.bIsNegative = false;
+	OutTimeSplit.bIsInfinite = false;
+	OutTimeSplit.bIsNaN = false;
+
+	if (std::isnan(InTime))
+	{
+		OutTimeSplit.bIsNaN = true;
+		return;
+	}
+
+	if (FMath::IsNearlyZero(Time, TimeUtils::Picosecond))
+	{
+		OutTimeSplit.bIsZero = true;
+		for (int32 Index = 0; Index < 8; ++Index)
+		{
+			OutTimeSplit.Units[Index] = 0;
+		}
+		return;
+	}
+
+	if (Time < 0.0)
+	{
+		Time = -Time;
+		OutTimeSplit.bIsNegative = true;
+	}
+
+	if (Time == DBL_MAX || Time == std::numeric_limits<double>::infinity())
+	{
+		OutTimeSplit.bIsInfinite = true;
+		for (int32 Index = 0; Index < 8; ++Index)
+		{
+			OutTimeSplit.Units[Index] = 0;
+		}
+		return;
+	}
+
+	bool bIsZero = true; // assume true and see if any split unit is != 0
+
+	const double Days = FMath::FloorToDouble(Time / TimeUtils::Day);
+	OutTimeSplit.Days = static_cast<int32>(Days);
+	if (OutTimeSplit.Days > 0)
+	{
+		Time -= Days * TimeUtils::Day;
+		bIsZero = false;
+	}
+
+	const double Hours = FMath::FloorToDouble(Time / TimeUtils::Hour);
+	OutTimeSplit.Hours = static_cast<int32>(Hours);
+	if (OutTimeSplit.Hours > 0)
+	{
+		Time -= Hours * TimeUtils::Hour;
+	}
+
+	const double Minutes = FMath::FloorToDouble(Time / TimeUtils::Minute);
+	OutTimeSplit.Minutes = static_cast<int32>(Minutes);
+	if (OutTimeSplit.Minutes > 0)
+	{
+		Time -= Minutes * TimeUtils::Minute;
+		bIsZero = false;
+	}
+
+	const double Seconds = FMath::FloorToDouble(Time / TimeUtils::Second);
+	OutTimeSplit.Seconds = static_cast<int32>(Seconds);
+	if (OutTimeSplit.Seconds > 0)
+	{
+		Time -= Seconds * TimeUtils::Second;
+		bIsZero = false;
+	}
+
+	const double Miliseconds = FMath::FloorToDouble(Time / TimeUtils::Milisecond);
+	OutTimeSplit.Miliseconds = static_cast<int32>(Miliseconds);
+	if (OutTimeSplit.Miliseconds > 0)
+	{
+		Time -= Miliseconds * TimeUtils::Milisecond;
+		bIsZero = false;
+	}
+
+	const double Microseconds = FMath::FloorToDouble(Time / TimeUtils::Microsecond);
+	OutTimeSplit.Microseconds = static_cast<int32>(Microseconds);
+	if (OutTimeSplit.Microseconds > 0)
+	{
+		Time -= Microseconds * TimeUtils::Microsecond;
+		bIsZero = false;
+	}
+
+	const double Nanoseconds = FMath::FloorToDouble(Time / TimeUtils::Nanosecond);
+	OutTimeSplit.Nanoseconds = static_cast<int32>(Nanoseconds);
+	if (OutTimeSplit.Nanoseconds > 0)
+	{
+		Time -= Nanoseconds * TimeUtils::Nanosecond;
+		bIsZero = false;
+	}
+
+	const double Picoseconds = FMath::FloorToDouble(Time / TimeUtils::Picosecond);
+	OutTimeSplit.Picoseconds = static_cast<int32>(Picoseconds);
+	if (OutTimeSplit.Picoseconds > 0)
+	{
+		Time -= Picoseconds * TimeUtils::Picosecond;
+		bIsZero = false;
+	}
+
+	OutTimeSplit.bIsZero = bIsZero;
+}
+
+////////////////////////////////////////////////////////////////////////////////////////////////////
+
+FString FormatTimeSplit(const FTimeSplit& InTimeSplit, const double Precision)
+{
+	if (InTimeSplit.bIsNaN)
+	{
+		return TEXT("NaN");
+	}
+
+	if (InTimeSplit.bIsZero)
+	{
+		return TEXT("0");
+	}
+
+	if (InTimeSplit.bIsInfinite)
+	{
+		return InTimeSplit.bIsNegative ? TEXT("-∞") : TEXT("∞");
+	}
+
+	TStringBuilder<64> StrBuilder;
+
+	if (InTimeSplit.bIsNegative)
+	{
+		StrBuilder.AppendChar(TEXT('-'));
+	}
+
+	bool bIsSpaceNeeded = false;
+
+	if (InTimeSplit.Days > 0)
+	{
+		StrBuilder.Appendf(TEXT("%dd"), InTimeSplit.Days);
+		bIsSpaceNeeded = true;
+	}
+
+	if (Precision >= TimeUtils::Day)
+	{
+		if (!bIsSpaceNeeded)
+		{
+			return TEXT("~0");
+		}
+		return FString(StrBuilder);
+	}
+
+	if (InTimeSplit.Hours > 0)
+	{
+		if (bIsSpaceNeeded)
+		{
+			StrBuilder.AppendChar(TEXT(' '));
+		}
+		StrBuilder.Appendf(TEXT("%dh"), InTimeSplit.Hours);
+		bIsSpaceNeeded = true;
+	}
+
+	if (Precision >= TimeUtils::Hour)
+	{
+		if (!bIsSpaceNeeded)
+		{
+			return TEXT("~0");
+		}
+		return FString(StrBuilder);
+	}
+
+	if (InTimeSplit.Minutes > 0)
+	{
+		if (bIsSpaceNeeded)
+		{
+			StrBuilder.AppendChar(TEXT(' '));
+		}
+		StrBuilder.Appendf(TEXT("%dm"), InTimeSplit.Minutes);
+		bIsSpaceNeeded = true;
+	}
+
+	if (Precision >= TimeUtils::Minute)
+	{
+		if (!bIsSpaceNeeded)
+		{
+			return TEXT("~0");
+		}
+		return FString(StrBuilder);
+	}
+
+	if (InTimeSplit.Seconds > 0)
+	{
+		if (bIsSpaceNeeded)
+		{
+			StrBuilder.AppendChar(TEXT(' '));
+		}
+		StrBuilder.Appendf(TEXT("%ds"), InTimeSplit.Seconds);
+		bIsSpaceNeeded = true;
+	}
+
+	if (Precision >= TimeUtils::Second)
+	{
+		if (!bIsSpaceNeeded)
+		{
+			return TEXT("~0");
+		}
+		return FString(StrBuilder);
+	}
+
+	if (InTimeSplit.Miliseconds > 0)
+	{
+		if (bIsSpaceNeeded)
+		{
+			StrBuilder.AppendChar(TEXT(' '));
+		}
+		StrBuilder.Appendf(TEXT("%dms"), InTimeSplit.Miliseconds);
+		bIsSpaceNeeded = true;
+	}
+
+	if (Precision >= TimeUtils::Milisecond)
+	{
+		if (!bIsSpaceNeeded)
+		{
+			return TEXT("~0");
+		}
+		return FString(StrBuilder);
+	}
+
+	if (InTimeSplit.Microseconds > 0)
+	{
+		if (bIsSpaceNeeded)
+		{
+			StrBuilder.AppendChar(TEXT(' '));
+		}
+		StrBuilder.Appendf(TEXT("%dµs"), InTimeSplit.Microseconds);
+		bIsSpaceNeeded = true;
+	}
+
+	if (Precision >= TimeUtils::Microsecond)
+	{
+		if (!bIsSpaceNeeded)
+		{
+			return TEXT("~0");
+		}
+		return FString(StrBuilder);
+	}
+
+	if (InTimeSplit.Nanoseconds > 0)
+	{
+		if (bIsSpaceNeeded)
+		{
+			StrBuilder.AppendChar(TEXT(' '));
+		}
+		StrBuilder.Appendf(TEXT("%dns"), InTimeSplit.Nanoseconds);
+		bIsSpaceNeeded = true;
+	}
+
+	if (Precision >= TimeUtils::Nanosecond)
+	{
+		if (!bIsSpaceNeeded)
+		{
+			return TEXT("~0");
+		}
+		return FString(StrBuilder);
+	}
+
+	if (InTimeSplit.Picoseconds > 0)
+	{
+		if (bIsSpaceNeeded)
+		{
+			StrBuilder.AppendChar(TEXT(' '));
+		}
+		StrBuilder.Appendf(TEXT("%dps"), InTimeSplit.Picoseconds);
+		bIsSpaceNeeded = true;
+	}
+
+	if (!bIsSpaceNeeded)
+	{
+		return TEXT("~0");
+	}
+
+	return FString(StrBuilder);
+}
+
+////////////////////////////////////////////////////////////////////////////////////////////////////
+
+FString FormatTimeSplit(const double Time, const double Precision)
+{
+	FTimeSplit TimeSplit;
+	SplitTime(Time, TimeSplit);
+	return FormatTimeSplit(TimeSplit, Precision);
+}
+
+////////////////////////////////////////////////////////////////////////////////////////////////////
+
+void TestTimeFormatting()
+{
+	double T1 = 1 * TimeUtils::Day + 2 * TimeUtils::Hour + 3 * TimeUtils::Minute + 4.567890123456789;
+	FString S1 = FormatTime(T1, TimeUtils::Day);
+	//UE_LOG(TraceInsights, Log, TEXT("D-T: %s"), *S1);
+	FString S2 = FormatTime(T1, TimeUtils::Hour);
+	//UE_LOG(TraceInsights, Log, TEXT("H-T: %s"), *S2);
+	FString S3 = FormatTime(T1, TimeUtils::Minute);
+	//UE_LOG(TraceInsights, Log, TEXT("M-T: %s"), *S3);
+	for (double P = 10.0; P >= TimeUtils::Picosecond; P /= 10.0)
+	{
+		FString SP = FormatTime(T1, P);
+		//UE_LOG(TraceInsights, Log, TEXT("P:%g T: %s"), P, *SP);
+	}
+}
+
+////////////////////////////////////////////////////////////////////////////////////////////////////
+
+void TestTimeAutoFormatting()
+{
+	static bool Once = false;
+	if (Once) return;
+	Once = true;
+
+	struct FTestTimeAutoData
+	{
+		double T;
+		const TCHAR* Msg;
+	};
+
+	static const FTestTimeAutoData Data[] =
+	{
+		{ TimeUtils::Minute,                TEXT("1m") },
+		{ TimeUtils::Second * 59.99,        TEXT("59.99s") },
+		{ TimeUtils::Second * 59.95,        TEXT("59.95s") },
+		{ 0, TEXT("[threshold 1m / 59.9s]") },
+		{ TimeUtils::Second * 59.94,        TEXT("59.94s") },
+		{ TimeUtils::Second * 59.9,         TEXT("59.9s") },
+		{ TimeUtils::Second * 10,           TEXT("10s") },
+		{ TimeUtils::Second,                TEXT("1s") },
+		{ TimeUtils::Milisecond * 999.99,   TEXT("999.99ms") },
+		{ TimeUtils::Milisecond * 999.95,   TEXT("999.95ms") },
+		{ 0, TEXT("[threshold 1s / 999.9 ms]") },
+		{ TimeUtils::Milisecond * 999.94,   TEXT("999.94ms") },
+		{ TimeUtils::Milisecond * 999.9,    TEXT("999.9ms") },
+		{ TimeUtils::Milisecond * 999,      TEXT("999ms") },
+		{ TimeUtils::Milisecond * 100,      TEXT("100ms") },
+		{ TimeUtils::Milisecond * 10,       TEXT("10ms") },
+		{ TimeUtils::Milisecond * 1.55,     TEXT("1.55ms") },
+		{ TimeUtils::Milisecond * 1.5,      TEXT("1.5ms") },
+		{ TimeUtils::Milisecond * 1.05,     TEXT("1.05ms") },
+		{ TimeUtils::Milisecond,            TEXT("1ms") },
+		{ TimeUtils::Microsecond * 999.99,  TEXT("999.99µs") },
+		{ TimeUtils::Microsecond * 999.95,  TEXT("999.95µs") },
+		{ 0, TEXT("[threshold 1 ms / 999.9 µs]") },
+		{ TimeUtils::Microsecond * 999.94,  TEXT("999.94µs") },
+		{ TimeUtils::Microsecond * 999.9,   TEXT("999.9µs") },
+		{ TimeUtils::Microsecond * 999,     TEXT("999µs") },
+		{ TimeUtils::Microsecond * 100,     TEXT("100µs") },
+		{ TimeUtils::Microsecond * 10,      TEXT("10µs") },
+		{ TimeUtils::Microsecond,           TEXT("1µs") },
+		{ TimeUtils::Nanosecond * 999.9,    TEXT("999.9ns") },
+		{ TimeUtils::Nanosecond * 999.5,    TEXT("999.5ns") },
+		{ 0, TEXT("[threshold 1 µs / 999 ns]") },
+		{ TimeUtils::Nanosecond * 999.4,    TEXT("999.4ns") },
+		{ TimeUtils::Nanosecond * 999,      TEXT("999ns") },
+		{ TimeUtils::Nanosecond * 100,      TEXT("100ns") },
+		{ TimeUtils::Nanosecond * 10,       TEXT("10ns") },
+		{ TimeUtils::Nanosecond,            TEXT("1ns") },
+		{ TimeUtils::Picosecond * 999.9,    TEXT("999.9ps") },
+		{ TimeUtils::Picosecond * 999.5,    TEXT("999.5ps") },
+		{ 0, TEXT("[threshold 1 ns / 999 ps]") },
+		{ TimeUtils::Picosecond * 999.4,    TEXT("999.4ps") },
+		{ TimeUtils::Picosecond * 999,      TEXT("999ps") },
+		{ TimeUtils::Picosecond * 100,      TEXT("100ps") },
+		{ TimeUtils::Picosecond * 10,       TEXT("10ps") },
+		{ TimeUtils::Picosecond,            TEXT("1ps") },
+		{ TimeUtils::Picosecond * 0.1,      TEXT("0.1ps") },
+	};
+	const int32 DataCount = sizeof(Data) / sizeof(FTestTimeAutoData);
+
+	for (int32 Index = 0; Index < DataCount; ++Index)
+	{
+		FString Str = FormatTimeAuto(Data[Index].T);
+		//UE_LOG(TraceInsights, Log, TEXT("%s : %s"), Data[Index].Msg, *Str);
+	}
+}
+
+////////////////////////////////////////////////////////////////////////////////////////////////////
+
+int32 GetNumDigits(const double Precision)
+{
+	//float Log10 = -log10(Precision);
+	float Log10 = -FMath::LogX(10.0f, static_cast<float>(Precision));
+	int32 D = (Log10 > 0) ? FMath::CeilToInt(Log10) : 0;
+	return D;
+}
+
+UE_DISABLE_OPTIMIZATION_SHIP
+int32 GetNumDigitsOptDisabled(const double Precision)
+{
+	//double Log10 = -log10(Precision);
+	float Log10 = -FMath::LogX(10.0f, static_cast<float>(Precision));
+	int32 D = (Log10 > 0) ? FMath::CeilToInt(Log10) : 0;
+	return D;
+}
+UE_ENABLE_OPTIMIZATION_SHIP
+
+UE_DISABLE_OPTIMIZATION_SHIP
+void TestOptimizationIssue()
+{
+	constexpr double Ns = 0.000000001;
+	int32 D1 = GetNumDigits(Ns);
+	int32 D2 = GetNumDigitsOptDisabled(Ns);
+	//UE_LOG(TraceInsights, Log, TEXT("D1 = %d, D2 = %d"), D1, D2);
+	ensure(D1 == 9); // 10 ?
+	ensure(D1 == D2);
+}
+UE_ENABLE_OPTIMIZATION_SHIP
+
+////////////////////////////////////////////////////////////////////////////////////////////////////
+
+} // namespace TimeUtils
+} // namespace Insights
+} // namespace UE
diff --git a/Engine/Plugins/Editor/AssetManagerEditor/Source/AssetManagerEditor/Private/Insights/Filter/ViewModels/Filters.h b/Engine/Plugins/Editor/AssetManagerEditor/Source/AssetManagerEditor/Private/Insights/Filter/ViewModels/Filters.h
new file mode 100644
--- /dev/null
+++ b/Engine/Plugins/Editor/AssetManagerEditor/Source/AssetManagerEditor/Private/Insights/Filter/ViewModels/Filters.h
@@ -0,0 +1,299 @@
+// Copyright Epic Games, Inc. All Rights Reserved.
+
+#pragma once
+
+#include <functional>
+
+#include "CoreMinimal.h"
+#include "Misc/TVariant.h"
+
+#include "Insights/Common/SimpleRtti.h"
+
+#define LOCTEXT_NAMESPACE "Filters"
+
+class FSpawnTabArgs;
+class SDockTab;
+class SWidget;
+
+namespace UE
+{
+namespace Insights
+{
+
+////////////////////////////////////////////////////////////////////////////////////////////////////
+
+enum class EFilterDataType : uint32
+{
+	Int64,
+	Double,
+	String,
+	StringInt64Pair, // Displayed as a string but translates to a Int64 key.
+};
+
+////////////////////////////////////////////////////////////////////////////////////////////////////
+
+enum class EFilterOperator : uint8
+{
+	Eq = 0, // Equals
+	NotEq = 1, // Not Equals
+	Lt = 2, // Less Than
+	Lte = 3, // Less than or equal to
+	Gt = 4, // Greater than
+	Gte = 5, // Greater than or equal to
+	Contains = 6,
+	NotContains = 7,
+};
+
+////////////////////////////////////////////////////////////////////////////////////////////////////
+
+class IFilterOperator
+{
+public:
+
+	virtual EFilterOperator GetKey() = 0;
+	virtual FString GetName() = 0;
+};
+
+////////////////////////////////////////////////////////////////////////////////////////////////////
+
+template<typename T>
+class FFilterOperator : public IFilterOperator
+{
+public:
+	typedef TFunction<bool(T, T)> OperatorFunc;
+
+	FFilterOperator(EFilterOperator InKey, FString InName, OperatorFunc InFunc)
+		: Func(InFunc)
+		, Key(InKey)
+		, Name(InName)
+	{
+	}
+	virtual ~FFilterOperator()
+	{
+	}
+
+	virtual EFilterOperator GetKey() override { return Key; }
+	virtual FString GetName() override { return Name; };
+
+	OperatorFunc Func;
+
+private:
+	EFilterOperator Key;
+	FString Name;
+};
+
+////////////////////////////////////////////////////////////////////////////////////////////////////
+
+enum class EFilterGroupOperator
+{
+	And = 0,
+	Or = 1,
+};
+
+////////////////////////////////////////////////////////////////////////////////////////////////////
+
+struct FFilterGroupOperator
+{
+	FFilterGroupOperator(EFilterGroupOperator InType, FText InName, FText InDesc)
+		: Type(InType)
+		, Name(InName)
+		, Desc(InDesc)
+	{
+	}
+
+	EFilterGroupOperator Type;
+	FText Name;
+	FText Desc;
+};
+
+////////////////////////////////////////////////////////////////////////////////////////////////////
+
+class IFilterValueConverter
+{
+public:
+	virtual bool Convert(const FString& Input, int64& Output, FText& OutError) const { unimplemented(); return false; }
+	virtual bool Convert(const FString& Input, double& Output, FText& OutError) const { unimplemented(); return false; }
+	virtual FText GetTooltipText() const { return FText(); }
+	virtual FText GetHintText() const { return FText(); }
+};
+
+////////////////////////////////////////////////////////////////////////////////////////////////////
+
+typedef TSharedPtr<const TArray<TSharedPtr<IFilterOperator>>> SupportedOperatorsArrayPtr;
+
+struct FFilter
+{
+	INSIGHTS_DECLARE_RTTI_BASE(FFilter)
+
+public:
+	FFilter(int32 InKey, FText InName, FText InDesc, EFilterDataType InDataType, SupportedOperatorsArrayPtr InSupportedOperators)
+		: Key(InKey)
+		, Name(InName)
+		, Desc(InDesc)
+		, DataType(InDataType)
+		, SupportedOperators(InSupportedOperators)
+	{
+	}
+	virtual ~FFilter()
+	{
+	}
+
+	SupportedOperatorsArrayPtr GetSupportedOperators() const { return SupportedOperators; }
+
+	int32 Key;
+	FText Name;
+	FText Desc;
+	EFilterDataType DataType;
+	TSharedPtr<IFilterValueConverter> Converter;
+	SupportedOperatorsArrayPtr SupportedOperators;
+};
+
+////////////////////////////////////////////////////////////////////////////////////////////////////
+
+struct FFilterWithSuggestions : FFilter
+{
+	typedef TFunction<void(const FString& /*Text*/, TArray<FString>& OutSuggestions)> GetSuggestionsCallback;
+
+	INSIGHTS_DECLARE_RTTI(FFilterWithSuggestions, FFilter)
+
+public:
+	FFilterWithSuggestions(int32 InKey, FText InName, FText InDesc, EFilterDataType InDataType, SupportedOperatorsArrayPtr InSupportedOperators)
+		: FFilter(InKey, InName, InDesc, InDataType, InSupportedOperators)
+	{
+	}
+
+	virtual ~FFilterWithSuggestions()
+	{
+	}
+
+	GetSuggestionsCallback Callback;
+};
+
+////////////////////////////////////////////////////////////////////////////////////////////////////
+
+class FFilterStorage
+{
+public:
+	FFilterStorage()
+	{
+		DoubleOperators = MakeShared<TArray<TSharedPtr<IFilterOperator>>>();
+		DoubleOperators->Add(StaticCastSharedRef<IFilterOperator>(MakeShared<FFilterOperator<double>>(EFilterOperator::Lt, TEXT("<"), std::less<>{})));
+		DoubleOperators->Add(StaticCastSharedRef<IFilterOperator>(MakeShared<FFilterOperator<double>>(EFilterOperator::Lte, TEXT("\u2264"), std::less_equal<>())));
+		DoubleOperators->Add(StaticCastSharedRef<IFilterOperator>(MakeShared<FFilterOperator<double>>(EFilterOperator::Eq, TEXT("="), std::equal_to<>())));
+		DoubleOperators->Add(StaticCastSharedRef<IFilterOperator>(MakeShared<FFilterOperator<double>>(EFilterOperator::Gt, TEXT(">"), std::greater<>())));
+		DoubleOperators->Add(StaticCastSharedRef<IFilterOperator>(MakeShared<FFilterOperator<double>>(EFilterOperator::Gte, TEXT("\u2265"), std::greater_equal<>())));
+
+		IntegerOperators = MakeShared<TArray<TSharedPtr<IFilterOperator>>>();
+		IntegerOperators->Add(StaticCastSharedRef<IFilterOperator>(MakeShared<FFilterOperator<int64>>(EFilterOperator::Lt, TEXT("<"), std::less<>{})));
+		IntegerOperators->Add(StaticCastSharedRef<IFilterOperator>(MakeShared<FFilterOperator<int64>>(EFilterOperator::Lte, TEXT("\u2264"), std::less_equal<>())));
+		IntegerOperators->Add(StaticCastSharedRef<IFilterOperator>(MakeShared<FFilterOperator<int64>>(EFilterOperator::Eq, TEXT("="), std::equal_to<>())));
+		IntegerOperators->Add(StaticCastSharedRef<IFilterOperator>(MakeShared<FFilterOperator<int64>>(EFilterOperator::Gt, TEXT(">"), std::greater<>())));
+		IntegerOperators->Add(StaticCastSharedRef<IFilterOperator>(MakeShared<FFilterOperator<int64>>(EFilterOperator::Gte, TEXT("\u2265"), std::greater_equal<>())));
+
+		StringOperators = MakeShared<TArray<TSharedPtr<IFilterOperator>>>();
+		StringOperators->Add(StaticCastSharedRef<IFilterOperator>(MakeShared<FFilterOperator<FString>>(EFilterOperator::Eq, TEXT("IS"), [](const FString& lhs, const FString& rhs) { return lhs.Equals(rhs); })));
+		StringOperators->Add(StaticCastSharedRef<IFilterOperator>(MakeShared<FFilterOperator<FString>>(EFilterOperator::NotEq, TEXT("IS NOT"), [](const FString& lhs, const FString& rhs) { return !lhs.Equals(rhs); })));
+		StringOperators->Add(StaticCastSharedRef<IFilterOperator>(MakeShared<FFilterOperator<FString>>(EFilterOperator::Contains, TEXT("CONTAINS"), [](const FString& lhs, const FString& rhs) { return lhs.Contains(rhs); })));
+		StringOperators->Add(StaticCastSharedRef<IFilterOperator>(MakeShared<FFilterOperator<FString>>(EFilterOperator::NotContains, TEXT("NOT CONTAINS"), [](const FString& lhs, const FString& rhs) { return !lhs.Contains(rhs); })));
+
+		FilterGroupOperators.Add(MakeShared<FFilterGroupOperator>(EFilterGroupOperator::And, LOCTEXT("AllOf", "All Of (AND)"), LOCTEXT("AllOfDesc", "All of the children must be true for the group to return true. Equivalent to an AND operation.")));
+		FilterGroupOperators.Add(MakeShared<FFilterGroupOperator>(EFilterGroupOperator::Or, LOCTEXT("AnyOf", "Any Of (OR)"), LOCTEXT("AnyOfDesc", "Any of the children must be true for the group to return true. Equivalent to an OR operation.")));
+	}
+
+	const TArray<TSharedPtr<FFilterGroupOperator>>& GetFilterGroupOperators() const { return FilterGroupOperators; }
+
+	TSharedPtr<TArray<TSharedPtr<IFilterOperator>>> GetDoubleOperators() const { return DoubleOperators; }
+	TSharedPtr<TArray<TSharedPtr<IFilterOperator>>> GetIntegerOperators() const { return IntegerOperators; }
+	TSharedPtr<TArray<TSharedPtr<IFilterOperator>>> GetStringOperators() const { return StringOperators; }
+
+private:
+	TArray<TSharedPtr<FFilterGroupOperator>> FilterGroupOperators;
+
+	TSharedPtr<TArray<TSharedPtr<IFilterOperator>>> DoubleOperators;
+	TSharedPtr<TArray<TSharedPtr<IFilterOperator>>> IntegerOperators;
+	TSharedPtr<TArray<TSharedPtr<IFilterOperator>>> StringOperators;
+};
+
+////////////////////////////////////////////////////////////////////////////////////////////////////
+
+class FFilterService
+{
+public:
+	FFilterService();
+	~FFilterService();
+
+	static void Initialize();
+	static void Shutdown();
+	static TSharedPtr<FFilterService> Get() { return Instance; }
+
+	void RegisterTabSpawner();
+	void UnregisterTabSpawner();
+	TSharedRef<SDockTab> SpawnTab(const FSpawnTabArgs& Args);
+
+	const TArray<TSharedPtr<FFilterGroupOperator>>& GetFilterGroupOperators() const { return FilterStorage.GetFilterGroupOperators(); }
+
+	TSharedPtr<TArray<TSharedPtr<IFilterOperator>>> GetDoubleOperators() const { return FilterStorage.GetDoubleOperators(); }
+	TSharedPtr<TArray<TSharedPtr<IFilterOperator>>> GetIntegerOperators() const { return FilterStorage.GetIntegerOperators(); }
+	TSharedPtr<TArray<TSharedPtr<IFilterOperator>>> GetStringOperators() const { return FilterStorage.GetStringOperators(); }
+
+	TSharedPtr<SWidget> CreateFilterConfiguratorWidget(TSharedPtr<class FFilterConfigurator> FilterConfiguratorViewModel);
+
+private:
+	static const FName FilterConfiguratorTabId;
+
+	static TSharedPtr<FFilterService> Instance;
+	FFilterStorage FilterStorage;
+
+	TSharedPtr<class SAdvancedFilter> PendingWidget;
+};
+
+////////////////////////////////////////////////////////////////////////////////////////////////////
+
+class FFilterContext
+{
+public:
+	typedef TVariant<double, int64, FString> ContextData;
+
+	template<typename T>
+	void AddFilterData(int32 Key, const T& InData)
+	{
+		ContextData VariantData;
+		VariantData.Set<T>(InData);
+		DataMap.Add(Key, VariantData);
+	}
+
+	template<typename T>
+	void SetFilterData(int32 Key, const T& InData)
+	{
+		DataMap[Key].Set<T>(InData);
+	}
+
+	template<typename T>
+	void GetFilterData(int32 Key, T& OutData) const
+	{
+		const ContextData* Data = DataMap.Find(Key);
+		check(Data);
+
+		check(Data->IsType<T>());
+		OutData = Data->Get<T>();
+	}
+
+	bool HasFilterData(int32 Key) const
+	{
+		return DataMap.Contains(Key);
+	}
+
+	bool GetReturnValueForUnsetFilters() const { return bReturnValueForUnsetFilters; }
+	void SetReturnValueForUnsetFilters(bool InValue) { bReturnValueForUnsetFilters = InValue; }
+
+private:
+	TMap<int32, ContextData> DataMap;
+	bool bReturnValueForUnsetFilters = true;
+};
+
+////////////////////////////////////////////////////////////////////////////////////////////////////
+
+} // namespace Insights
+} // namespace UE
+
+#undef LOCTEXT_NAMESPACE
diff --git a/Engine/Plugins/Editor/AssetManagerEditor/Source/AssetManagerEditor/Private/Insights/Table/ViewModels/BaseTreeNode.cpp b/Engine/Plugins/Editor/AssetManagerEditor/Source/AssetManagerEditor/Private/Insights/Table/ViewModels/BaseTreeNode.cpp
new file mode 100644
--- /dev/null
+++ b/Engine/Plugins/Editor/AssetManagerEditor/Source/AssetManagerEditor/Private/Insights/Table/ViewModels/BaseTreeNode.cpp
@@ -0,0 +1,74 @@
+// Copyright Epic Games, Inc. All Rights Reserved.
+
+#include "BaseTreeNode.h"
+
+// Insights
+#include "Insights/Table/ViewModels/TableCellValueSorter.h"
+
+#define LOCTEXT_NAMESPACE "UE::Insights::FBaseTreeNode"
+
+namespace UE
+{
+namespace Insights
+{
+
+INSIGHTS_IMPLEMENT_RTTI(FBaseTreeNode)
+
+FBaseTreeNode::FGroupNodeData FBaseTreeNode::DefaultGroupData;
+
+////////////////////////////////////////////////////////////////////////////////////////////////////
+
+const FText FBaseTreeNode::GetDisplayName() const
+{
+	return FText::FromName(GetName());
+}
+
+////////////////////////////////////////////////////////////////////////////////////////////////////
+
+const FText FBaseTreeNode::GetExtraDisplayName() const
+{
+	if (IsGroup())
+	{
+		const int32 NumChildren = GroupData->Children.Num();
+		const int32 NumFilteredChildren = GroupData->FilteredChildren.Num();
+
+		if (NumFilteredChildren == NumChildren)
+		{
+			return FText::Format(LOCTEXT("TreeNodeGroup_ExtraText_Fmt1", "({0})"), FText::AsNumber(NumChildren));
+		}
+		else
+		{
+			return FText::Format(LOCTEXT("TreeNodeGroup_ExtraText_Fmt2", "({0} / {1})"), FText::AsNumber(NumFilteredChildren), FText::AsNumber(NumChildren));
+		}
+	}
+
+	return FText::GetEmpty();
+}
+
+////////////////////////////////////////////////////////////////////////////////////////////////////
+
+bool FBaseTreeNode::HasExtraDisplayName() const
+{
+	return IsGroup();
+}
+
+////////////////////////////////////////////////////////////////////////////////////////////////////
+
+void FBaseTreeNode::SortChildrenAscending(const ITableCellValueSorter& Sorter)
+{
+	Sorter.Sort(GroupData->Children, ESortMode::Ascending);
+}
+
+////////////////////////////////////////////////////////////////////////////////////////////////////
+
+void FBaseTreeNode::SortChildrenDescending(const ITableCellValueSorter& Sorter)
+{
+	Sorter.Sort(GroupData->Children, ESortMode::Descending);
+}
+
+////////////////////////////////////////////////////////////////////////////////////////////////////
+
+} // namespace Insights
+} // namespace UE
+
+#undef LOCTEXT_NAMESPACE
diff --git a/Engine/Plugins/Editor/AssetManagerEditor/Source/AssetManagerEditor/Private/Insights/Table/ViewModels/BaseTreeNode.h b/Engine/Plugins/Editor/AssetManagerEditor/Source/AssetManagerEditor/Private/Insights/Table/ViewModels/BaseTreeNode.h
new file mode 100644
--- /dev/null
+++ b/Engine/Plugins/Editor/AssetManagerEditor/Source/AssetManagerEditor/Private/Insights/Table/ViewModels/BaseTreeNode.h
@@ -0,0 +1,308 @@
+// Copyright Epic Games, Inc. All Rights Reserved.
+
+#pragma once
+
+#include "CoreMinimal.h"
+
+#include "Insights/Common/SimpleRtti.h"
+
+namespace UE
+{
+namespace Insights
+{
+
+class ITableCellValueSorter;
+
+////////////////////////////////////////////////////////////////////////////////////////////////////
+
+class FBaseTreeNode;
+
+/** Type definition for shared pointers to instances of FBaseTreeNode. */
+typedef TSharedPtr<class FBaseTreeNode> FBaseTreeNodePtr;
+
+/** Type definition for shared references to instances of FBaseTreeNode. */
+typedef TSharedRef<class FBaseTreeNode> FBaseTreeNodeRef;
+
+/** Type definition for shared references to const instances of FBaseTreeNode. */
+typedef TSharedRef<const class FBaseTreeNode> FBaseTreeNodeRefConst;
+
+/** Type definition for weak references to instances of FBaseTreeNode. */
+typedef TWeakPtr<class FBaseTreeNode> FBaseTreeNodeWeak;
+
+////////////////////////////////////////////////////////////////////////////////////////////////////
+
+class FBaseTreeNode : public TSharedFromThis<FBaseTreeNode>
+{
+	INSIGHTS_DECLARE_RTTI_BASE(FBaseTreeNode)
+
+protected:
+	struct FGroupNodeData
+	{
+		/** Children of the group node. */
+		TArray<FBaseTreeNodePtr> Children;
+
+		/** Filtered children of the group node. */
+		TArray<FBaseTreeNodePtr> FilteredChildren;
+
+		/** Descriptive text for the group node to display in a tooltip. */
+		FText Tooltip;
+
+		/** A generic pointer to a data context for the group node. Can be nullptr. */
+		void* Context = nullptr;
+
+		/** Whether the group node should be expanded or not. */
+		bool bIsExpanded = false;
+	};
+
+public:
+	/** Initialization constructor for the node. */
+	FBaseTreeNode(const FName InName, bool bInIsGroup)
+		: DefaultSortOrder(0)
+		, Name(InName)
+		, GroupData(bInIsGroup ? new FGroupNodeData() : &DefaultGroupData)
+	{
+	}
+
+	virtual ~FBaseTreeNode()
+	{
+		if (GroupData != &DefaultGroupData)
+		{
+			delete GroupData;
+			GroupData = nullptr;
+		}
+	}
+
+	uint32 GetDefaultSortOrder() const
+	{
+		return DefaultSortOrder;
+	}
+	void SetDefaultSortOrder(uint32 Order)
+	{
+		DefaultSortOrder = Order;
+	}
+
+	/**
+	 * @return a name of this node.
+	 */
+	const FName& GetName() const
+	{
+		return Name;
+	}
+
+	/**
+	 * @return a name of this node to display in tree view.
+	 */
+	virtual const FText GetDisplayName() const;
+
+	/**
+	 * @return a name suffix for this node to display in tree view. Ex.: group nodes may include additional info, like the number of visible / total children.
+	 */
+	virtual const FText GetExtraDisplayName() const;
+
+	virtual bool HasExtraDisplayName() const;
+
+	/**
+	 * @return a descriptive text for this node to display in a tooltip.
+	 */
+	virtual const FText GetTooltip() const
+	{
+		return GroupData->Tooltip;
+	}
+
+	/**
+	 * Sets a descriptive text for this node to display in a tooltip.
+	 */
+	virtual void SetTooltip(const FText& InTooltip)
+	{
+		if (IsGroup())
+		{
+			GroupData->Tooltip = InTooltip;
+		}
+	}
+
+	/**
+	 * @return a pointer to a data context for this node.
+	 */
+	virtual void* GetContext() const
+	{
+		return GroupData->Context;
+	}
+
+	/**
+	 * Sets a pointer to a data context for this node.
+	 */
+	virtual void SetContext(void* InContext)
+	{
+		if (IsGroup())
+		{
+			GroupData->Context = InContext;
+		}
+	}
+
+	/**
+	 * @return true, if this node is a group node.
+	 */
+	bool IsGroup() const
+	{
+		return GroupData != &DefaultGroupData;
+	}
+
+	/**
+	 * Initializes the group data, allowing this node to accept children.
+	 */
+	void InitGroupData()
+	{
+		if (GroupData == &DefaultGroupData)
+		{
+			GroupData = new FGroupNodeData();
+		}
+	}
+
+	/**
+	 * @return a const reference to the child nodes of this group.
+	 */
+	FORCEINLINE_DEBUGGABLE const TArray<FBaseTreeNodePtr>& GetChildren() const
+	{
+		return GroupData->Children;
+	}
+
+	/**
+	 * @return a const reference to the child nodes that should be visible to the UI based on filtering.
+	 */
+	FORCEINLINE_DEBUGGABLE const TArray<FBaseTreeNodePtr>& GetFilteredChildren() const
+	{
+		return GroupData->FilteredChildren;
+	}
+
+	/**
+	 * @return a weak reference to the group of this node, may be invalid.
+	 */
+	FBaseTreeNodeWeak GetGroupPtr() const
+	{
+		return GroupPtr;
+	}
+
+	/**
+	 * @return a name of the group node that this node belongs to.
+	 */
+	const FName& GetGroupName() const
+	{
+		return GroupPtr.Pin()->GetName();
+	}
+
+	virtual bool IsFiltered() const
+	{
+		return false;
+	}
+
+	void SortChildrenAscending(const ITableCellValueSorter& Sorter);
+	void SortChildrenDescending(const ITableCellValueSorter& Sorter);
+
+	template <class PREDICATE_CLASS>
+	void SortChildren(const PREDICATE_CLASS& Predicate)
+	{
+		GroupData->Children.Sort(Predicate);
+	}
+
+	/** Adds specified child to the children and sets group for it. */
+	FORCEINLINE_DEBUGGABLE void AddChildAndSetGroupPtr(const FBaseTreeNodePtr& ChildPtr)
+	{
+		ChildPtr->GroupPtr = AsShared();
+		GroupData->Children.Add(ChildPtr);
+	}
+
+	void SetGroupPtrForAllChildren()
+	{
+		FBaseTreeNodeWeak ThisNode = AsShared();
+		for (FBaseTreeNodePtr& NodePtr : GroupData->Children)
+		{
+			NodePtr->GroupPtr = ThisNode;
+		}
+	}
+
+	void ResetGroupPtrForAllChildren()
+	{
+		for (FBaseTreeNodePtr& NodePtr : GroupData->Children)
+		{
+			NodePtr->GroupPtr = nullptr;
+		}
+	}
+
+	/** Clears children. */
+	FORCEINLINE_DEBUGGABLE void ClearChildren(int32 NewSize = 0)
+	{
+		ResetGroupPtrForAllChildren();
+		GroupData->Children.Reset(NewSize);
+	}
+
+	FORCEINLINE_DEBUGGABLE void SwapChildren(TArray<FBaseTreeNodePtr>& NewChildren)
+	{
+		ResetGroupPtrForAllChildren();
+		Swap(NewChildren, GroupData->Children);
+		SetGroupPtrForAllChildren();
+	}
+
+	FORCEINLINE_DEBUGGABLE void SwapChildrenFast(TArray<FBaseTreeNodePtr>& NewChildren)
+	{
+		Swap(NewChildren, GroupData->Children);
+	}
+
+	/** Adds specified child to the filtered children. */
+	FORCEINLINE_DEBUGGABLE void AddFilteredChild(const FBaseTreeNodePtr& ChildPtr)
+	{
+		GroupData->FilteredChildren.Add(ChildPtr);
+	}
+
+	/** Clears filtered children. */
+	void ClearFilteredChildren(int32 NewSize = 0)
+	{
+		GroupData->FilteredChildren.Reset(NewSize);
+	}
+
+	bool IsExpanded() const
+	{
+		return GroupData->bIsExpanded;
+	}
+	void SetExpansion(bool bOnOff)
+	{
+		if (IsGroup())
+		{
+			GroupData->bIsExpanded = bOnOff;
+		}
+	}
+
+protected:
+	/**
+	 * @return a reference to the child nodes of this group.
+	 */
+	FORCEINLINE_DEBUGGABLE TArray<FBaseTreeNodePtr>& GetChildrenMutable()
+	{
+		return GroupData->Children;
+	}
+
+	void SetGroupPtr(FBaseTreeNodePtr InGroupPtr)
+	{
+		GroupPtr = InGroupPtr;
+	}
+
+private:
+	/** The default sort order. Index used to optimize sorting. */
+	int32 DefaultSortOrder;
+
+	/** The name of this node. */
+	const FName Name;
+
+	/** A weak pointer to the group/parent of this node. */
+	FBaseTreeNodeWeak GroupPtr;
+
+	/** The struct containing properties of a group node. It is allocated only for group nodes. */
+	FGroupNodeData* GroupData;
+
+	/** The only group data for "not a group" nodes. */
+	static FGroupNodeData DefaultGroupData;
+};
+
+////////////////////////////////////////////////////////////////////////////////////////////////////
+
+} // namespace Insights
+} // namespace UE
diff --git a/Engine/Plugins/Editor/AssetManagerEditor/Source/AssetManagerEditor/Private/Insights/Table/ViewModels/Table.h b/Engine/Plugins/Editor/AssetManagerEditor/Source/AssetManagerEditor/Private/Insights/Table/ViewModels/Table.h
new file mode 100644
--- /dev/null
+++ b/Engine/Plugins/Editor/AssetManagerEditor/Source/AssetManagerEditor/Private/Insights/Table/ViewModels/Table.h
@@ -0,0 +1,81 @@
+// Copyright Epic Games, Inc. All Rights Reserved.
+
+#pragma once
+
+#include "CoreMinimal.h"
+
+namespace UE
+{
+namespace Insights
+{
+
+class FTableColumn;
+
+////////////////////////////////////////////////////////////////////////////////////////////////////
+/**
+ * Table View Model.
+ * View model class for the STableListView and STableTreeView widgets.
+ */
+class FTable : public TSharedFromThis<FTable>
+{
+public:
+	FTable();
+	virtual ~FTable();
+
+	const FText& GetDisplayName() const { return DisplayName; }
+	void SetDisplayName(const FText& InDisplayName) { DisplayName = InDisplayName; }
+
+	const FText& GetDescription() const { return Description; }
+	void SetDescription(const FText& InDescription) { Description = InDescription; }
+
+	virtual void Reset();
+
+	int32 GetColumnCount() const { return Columns.Num(); }
+	bool IsValid() const { return Columns.Num() > 0; }
+
+	const TArray<TSharedRef<FTableColumn>>& GetColumns() const { return Columns; }
+	void SetColumns(const TArray<TSharedRef<FTableColumn>>& InColumns);
+
+	void GetVisibleColumns(TArray<TSharedRef<FTableColumn>>& InArray) const;
+
+	TSharedRef<FTableColumn> FindColumnChecked(const FName& ColumnId) const
+	{
+		return ColumnIdToPtrMapping.FindChecked(ColumnId);
+	}
+
+	TSharedPtr<FTableColumn> FindColumn(const FName& ColumnId) const
+	{
+		const TSharedRef<FTableColumn>* const ColumnRefPtr = ColumnIdToPtrMapping.Find(ColumnId);
+		if (ColumnRefPtr != nullptr)
+		{
+			return *ColumnRefPtr;
+		}
+		return nullptr;
+	}
+
+	int32 GetColumnPositionIndex(const FName& ColumnId) const;
+
+	void GetVisibleColumnsData(const TArray<TSharedPtr<class FBaseTreeNode>>& InNodes, const FName& LogListingName, TCHAR Separator, bool bIncludeHeaders, FString& OutData) const;
+
+	static const FName GetHierarchyColumnId();
+
+protected:
+	void ResetColumns() { Columns.Reset(); }
+	void AddColumn(TSharedRef<FTableColumn> Column);
+	void AddHierarchyColumn(int32 ColumnIndex, const TCHAR* ColumnName);
+
+private:
+	FText DisplayName;
+	FText Description;
+
+	/** All available columns. */
+	TArray<TSharedRef<FTableColumn>> Columns;
+
+	/** Mapping between column Ids and FTableColumn shared refs. */
+	TMap<FName, TSharedRef<FTableColumn>> ColumnIdToPtrMapping;
+};
+
+////////////////////////////////////////////////////////////////////////////////////////////////////
+
+} // namespace Insights
+} // namespace UE
diff --git a/Engine/Plugins/Editor/AssetManagerEditor/Source/AssetManagerEditor/Private/Insights/Table/ViewModels/TableCellValueFormatter.h b/Engine/Plugins/Editor/AssetManagerEditor/Source/AssetManagerEditor/Private/Insights/Table/ViewModels/TableCellValueFormatter.h
new file mode 100644
--- /dev/null
+++ b/Engine/Plugins/Editor/AssetManagerEditor/Source/AssetManagerEditor/Private/Insights/Table/ViewModels/TableCellValueFormatter.h
@@ -0,0 +1,207 @@
+// Copyright Epic Games, Inc. All Rights Reserved.
+
+#pragma once
+
+#include "CoreMinimal.h"
+
+#include "Layout/Visibility.h"
+
+#include "Insights/Table/ViewModels/TableCellValue.h"
+
+class IToolTip;
+
+namespace UE
+{
+namespace Insights
+{
+
+class FBaseTreeNode;
+class FTableColumn;
+
+////////////////////////////////////////////////////////////////////////////////////////////////////
+
+class ITableCellValueFormatter
+{
+public:
+	virtual FText FormatValue(const TOptional<FTableCellValue>& InValue) const = 0;
+	virtual FText FormatValueForTooltip(const TOptional<FTableCellValue>& InValue) const = 0;
+
+	virtual FText FormatValue(const FTableColumn& Column, const FBaseTreeNode& Node) const = 0;
+	virtual FText FormatValueForTooltip(const FTableColumn& Column, const FBaseTreeNode& Node) const = 0;
+	virtual FText FormatValueForGrouping(const FTableColumn& Column, const FBaseTreeNode& Node) const = 0;
+
+	virtual TSharedPtr<IToolTip> GetCustomTooltip(const FTableColumn& Column, const FBaseTreeNode& Node) const = 0;
+};
+
+////////////////////////////////////////////////////////////////////////////////////////////////////
+
+class FTableCellValueFormatter : public ITableCellValueFormatter
+{
+public:
+	FTableCellValueFormatter() {}
+	virtual ~FTableCellValueFormatter() {}
+
+	virtual FText FormatValue(const TOptional<FTableCellValue>& InValue) const override { return FText::GetEmpty(); }
+	virtual FText FormatValueForTooltip(const TOptional<FTableCellValue>& InValue) const override { return FormatValue(InValue); }
+
+	virtual FText FormatValue(const FTableColumn& Column, const FBaseTreeNode& Node) const override; // { return FormatValue(Column.GetValue(Node)); }
+	virtual FText FormatValueForTooltip(const FTableColumn& Column, const FBaseTreeNode& Node) const override; // { return FormatValueForTooltip(Column.GetValue(Node)); }
+	virtual FText FormatValueForGrouping(const FTableColumn& Column, const FBaseTreeNode& Node) const override; // { return FormatValueForTooltip(Column.GetValue(Node)); }
+
+	virtual TSharedPtr<IToolTip> GetCustomTooltip(const FTableColumn& Column, const FBaseTreeNode& Node) const override;
+
+	static EVisibility GetTooltipVisibility();
+};
+
+////////////////////////////////////////////////////////////////////////////////////////////////////
+
+class FTextValueFormatter : public FTableCellValueFormatter
+{
+public:
+	virtual FText FormatValue(const TOptional<FTableCellValue>& InValue) const override
+	{
+		if (InValue.IsSet())
+		{
+			return InValue.GetValue().GetText();
+		}
+		return FText::GetEmpty();
+	}
+};
+
+////////////////////////////////////////////////////////////////////////////////////////////////////
+
+class FAsTextValueFormatter : public FTableCellValueFormatter
+{
+public:
+	virtual FText FormatValue(const TOptional<FTableCellValue>& InValue) const override
+	{
+		if (InValue.IsSet())
+		{
+			return InValue.GetValue().AsText();
+		}
+		return FText::GetEmpty();
+	}
+};
+
+////////////////////////////////////////////////////////////////////////////////////////////////////
+
+class FBoolValueFormatterAsTrueFalse : public FTableCellValueFormatter
+{
+public:
+	virtual FText FormatValue(const TOptional<FTableCellValue>& InValue) const override;
+};
+
+////////////////////////////////////////////////////////////////////////////////////////////////////
+
+class FBoolValueFormatterAsOnOff : public FTableCellValueFormatter
+{
+public:
+	virtual FText FormatValue(const TOptional<FTableCellValue>& InValue) const override;
+};
+
+////////////////////////////////////////////////////////////////////////////////////////////////////
+
+class FInt64ValueFormatterAsNumber : public FTableCellValueFormatter
+{
+public:
+	virtual FText FormatValue(const TOptional<FTableCellValue>& InValue) const override
+	{
+		return InValue.IsSet() ? FText::AsNumber(InValue.GetValue().Int64) : FText::GetEmpty();
+	}
+};
+
+////////////////////////////////////////////////////////////////////////////////////////////////////
+
+class FInt64ValueFormatterAsUInt32InfinteNumber : public FTableCellValueFormatter
+{
+public:
+	virtual FText FormatValue(const TOptional<FTableCellValue>& InValue) const override;
+	virtual FText FormatValueForTooltip(const TOptional<FTableCellValue>& InValue) const override;
+};
+
+////////////////////////////////////////////////////////////////////////////////////////////////////
+
+class FInt64ValueFormatterAsHex32 : public FTableCellValueFormatter
+{
+public:
+	virtual FText FormatValue(const TOptional<FTableCellValue>& InValue) const override
+	{
+		return InValue.IsSet() ? FText::FromString(FString::Printf(TEXT("0x%08X"), static_cast<uint32>(InValue.GetValue().Int64))) : FText::GetEmpty();
+	}
+};
+
+////////////////////////////////////////////////////////////////////////////////////////////////////
+
+class FInt64ValueFormatterAsHex64 : public FTableCellValueFormatter
+{
+public:
+	virtual FText FormatValue(const TOptional<FTableCellValue>& InValue) const override
+	{
+		return InValue.IsSet() ? FText::FromString(FString::Printf(TEXT("0x%016llX"), static_cast<uint64>(InValue.GetValue().Int64))) : FText::GetEmpty();
+	}
+};
+
+////////////////////////////////////////////////////////////////////////////////////////////////////
+
+class FInt64ValueFormatterAsMemory : public FTableCellValueFormatter
+{
+public:
+	virtual FText FormatValue(const TOptional<FTableCellValue>& InValue) const override;
+	virtual FText FormatValueForTooltip(const TOptional<FTableCellValue>& InValue) const override;
+};
+
+////////////////////////////////////////////////////////////////////////////////////////////////////
+
+class FFloatValueFormatterAsNumber : public FTableCellValueFormatter
+{
+public:
+	virtual FText FormatValue(const TOptional<FTableCellValue>& InValue) const override;
+};
+
+////////////////////////////////////////////////////////////////////////////////////////////////////
+
+class FFloatValueFormatterAsTimeAuto : public FTableCellValueFormatter
+{
+public:
+	virtual FText FormatValue(const TOptional<FTableCellValue>& InValue) const override;
+	virtual FText FormatValueForTooltip(const TOptional<FTableCellValue>& InValue) const override;
+};
+
+////////////////////////////////////////////////////////////////////////////////////////////////////
+
+class FDoubleValueFormatterAsNumber : public FTableCellValueFormatter
+{
+public:
+	virtual FText FormatValue(const TOptional<FTableCellValue>& InValue) const override;
+};
+
+////////////////////////////////////////////////////////////////////////////////////////////////////
+
+class FDoubleValueFormatterAsTimeAuto : public FTableCellValueFormatter
+{
+public:
+	virtual FText FormatValue(const TOptional<FTableCellValue>& InValue) const override;
+	virtual FText FormatValueForTooltip(const TOptional<FTableCellValue>& InValue) const override;
+};
+
+////////////////////////////////////////////////////////////////////////////////////////////////////
+
+class FDoubleValueFormatterAsTimeMs : public FTableCellValueFormatter
+{
+public:
+	virtual FText FormatValue(const TOptional<FTableCellValue>& InValue) const override;
+	virtual FText FormatValueForTooltip(const TOptional<FTableCellValue>& InValue) const override;
+};
+
+////////////////////////////////////////////////////////////////////////////////////////////////////
+
+class FCStringValueFormatterAsText : public FTableCellValueFormatter
+{
+public:
+	virtual FText FormatValue(const TOptional<FTableCellValue>& InValue) const override;
+};
+
+////////////////////////////////////////////////////////////////////////////////////////////////////
+
+} // namespace Insights
+} // namespace UE
diff --git a/Engine/Plugins/Editor/AssetManagerEditor/Source/AssetManagerEditor/Private/Insights/Table/ViewModels/TableTreeNode.cpp b/Engine/Plugins/Editor/AssetManagerEditor/Source/AssetManagerEditor/Private/Insights/Table/ViewModels/TableTreeNode.cpp
new file mode 100644
--- /dev/null
+++ b/Engine/Plugins/Editor/AssetManagerEditor/Source/AssetManagerEditor/Private/Insights/Table/ViewModels/TableTreeNode.cpp
@@ -0,0 +1,17 @@
+// Copyright Epic Games, Inc. All Rights Reserved.
+
+#include "TableTreeNode.h"
+
+#define LOCTEXT_NAMESPACE "UE::Insights::FTableTreeNode"
+
+namespace UE
+{
+namespace Insights
+{
+
+INSIGHTS_IMPLEMENT_RTTI(FTableTreeNode)
+
+} // namespace Insights
+} // namespace UE
+
+#undef LOCTEXT_NAMESPACE
diff --git a/Engine/Plugins/Editor/AssetManagerEditor/Source/AssetManagerEditor/Private/Insights/Table/ViewModels/TableTreeNode.h b/Engine/Plugins/Editor/AssetManagerEditor/Source/AssetManagerEditor/Private/Insights/Table/ViewModels/TableTreeNode.h
new file mode 100644
--- /dev/null
+++ b/Engine/Plugins/Editor/AssetManagerEditor/Source/AssetManagerEditor/Private/Insights/Table/ViewModels/TableTreeNode.h
@@ -0,0 +1,123 @@
+// Copyright Epic Games, Inc. All Rights Reserved.
+
+#pragma once
+
+#include "CoreMinimal.h"
+
+// Insights
+#include "Insights/Table/ViewModels/BaseTreeNode.h"
+#include "Insights/Table/ViewModels/TableCellValue.h"
+
+namespace UE
+{
+namespace Insights
+{
+
+class FTable;
+
+struct FTableRowId
+{
+	static constexpr int32 InvalidRowIndex = -1;
+
+	FTableRowId(int32 InRowIndex) : RowIndex(InRowIndex) {}
+
+	bool HasValidIndex() const { return RowIndex >= 0; }
+
+	int32 RowIndex;
+};
+
+////////////////////////////////////////////////////////////////////////////////////////////////////
+
+class FTableTreeNode;
+
+/** Type definition for shared pointers to instances of FTableTreeNode. */
+typedef TSharedPtr<class FTableTreeNode> FTableTreeNodePtr;
+
+/** Type definition for shared references to instances of FTableTreeNode. */
+typedef TSharedRef<class FTableTreeNode> FTableTreeNodeRef;
+
+/** Type definition for shared references to const instances of FTableTreeNode. */
+typedef TSharedRef<const class FTableTreeNode> FTableTreeNodeRefConst;
+
+/** Type definition for weak references to instances of FTableTreeNode. */
+typedef TWeakPtr<class FTableTreeNode> FTableTreeNodeWeak;
+
+////////////////////////////////////////////////////////////////////////////////////////////////////
+/**
+ * Table Tree Node View Model.
+ * Class used to store information about a generic table tree node (used in STableTreeView).
+ */
+class FTableTreeNode : public FBaseTreeNode
+{
+	INSIGHTS_DECLARE_RTTI(FTableTreeNode, FBaseTreeNode)
+
+public:
+	/** Initialization constructor for a table record node. */
+	FTableTreeNode(const FName InName, TWeakPtr<FTable> InParentTable, int32 InRowIndex)
+		: FBaseTreeNode(InName, false)
+		, ParentTable(InParentTable)
+		, RowId(InRowIndex)
+		, AggregatedValues(nullptr)
+	{
+	}
+
+	/** Initialization constructor for a group node. */
+	FTableTreeNode(const FName InGroupName, TWeakPtr<FTable> InParentTable)
+		: FBaseTreeNode(InGroupName, true)
+		, ParentTable(InParentTable)
+		, RowId(FTableRowId::InvalidRowIndex)
+		, AggregatedValues(nullptr)
+	{
+	}
+
+	virtual ~FTableTreeNode()
+	{
+		CleanupAggregatedValues();
+	}
+
+	const TWeakPtr<FTable>& GetParentTable() const { return ParentTable; }
+	FTableRowId GetRowId() const { return RowId; }
+	int32 GetRowIndex() const { return RowId.RowIndex; }
+
+	void InitAggregatedValues()
+	{
+		if (!AggregatedValues)
+		{
+			AggregatedValues = new TMap<FName, FTableCellValue>();
+		}
+	}
+
+	void CleanupAggregatedValues()
+	{
+		if (AggregatedValues)
+		{
+			delete AggregatedValues;
+			AggregatedValues = nullptr;
+		}
+	}
+
+	void ResetAggregatedValues() { CleanupAggregatedValues(); }
+	void ResetAggregatedValues(const FName& ColumnId) { if (AggregatedValues) { AggregatedValues->Remove(ColumnId); } }
+	bool HasAggregatedValue(const FName& ColumnId) const { return AggregatedValues && AggregatedValues->Contains(ColumnId); }
+	const FTableCellValue* FindAggregatedValue(const FName& ColumnId) const { return AggregatedValues ? AggregatedValues->Find(ColumnId) : nullptr; }
+	const FTableCellValue& GetAggregatedValue(const FName& ColumnId) const { return AggregatedValues->FindChecked(ColumnId); }
+	void AddAggregatedValue(const FName& ColumnId, const FTableCellValue& Value) { InitAggregatedValues(); AggregatedValues->Add(ColumnId, Value); }
+	void SetAggregatedValue(const FName& ColumnId, const FTableCellValue& Value) { InitAggregatedValues(); (*AggregatedValues)[ColumnId] = Value; }
+
+	virtual bool IsFiltered() const override { return bIsFiltered; }
+	virtual void SetIsFiltered(bool InValue) { bIsFiltered = InValue; }
+
+protected:
+	TWeakPtr<FTable> ParentTable;
+	FTableRowId RowId;
+
+	TMap<FName, FTableCellValue>* AggregatedValues;
+
+private:
+	bool bIsFiltered = false;
+};
+
+////////////////////////////////////////////////////////////////////////////////////////////////////
+
+} // namespace Insights
+} // namespace UE
