AssetManagerEditor: Added Insights' TreeView implementation, including "advanced filtering" support and needed common code.

diff --git a/Engine/Plugins/Editor/AssetManagerEditor/Source/AssetManagerEditor/Private/Insights/Table/Widgets/STableTreeView.h b/Engine/Plugins/Editor/AssetManagerEditor/Source/AssetManagerEditor/Private/Insights/Table/Widgets/STableTreeView.h
new file mode 100644
--- /dev/null
+++ b/Engine/Plugins/Editor/AssetManagerEditor/Source/AssetManagerEditor/Private/Insights/Table/Widgets/STableTreeView.h
@@ -0,0 +1,795 @@
+// Copyright Epic Games, Inc. All Rights Reserved.
+
+#pragma once
+
+#include "CoreMinimal.h"
+#include "Misc/EnumClassFlags.h"
+#include "Misc/FilterCollection.h"
+#include "Misc/TextFilter.h"
+#include "SlateFwd.h"
+#include "Widgets/DeclarativeSyntaxSupport.h"
+#include "Widgets/Input/SComboBox.h"
+#include "Widgets/Navigation/SBreadcrumbTrail.h"
+#include "Widgets/SCompoundWidget.h"
+#include "Widgets/SWidget.h"
+#include "Widgets/Views/SHeaderRow.h"
+#include "Widgets/Views/STableRow.h"
+#include "Widgets/Views/STableViewBase.h"
+#include "Widgets/Views/STreeView.h"
+
+// Insights
+#include "Insights/Common/AsyncOperationProgress.h"
+#include "Insights/Common/AsyncWorkUtils.h"
+#include "Insights/Common/Stopwatch.h"
+#include "Insights/Filter/ViewModels/Filters.h"
+#include "Insights/Table/ViewModels/TableColumn.h"
+#include "Insights/Table/ViewModels/TableTreeNode.h"
+
+#include <atomic>
+
+class FMenuBuilder;
+class FUICommandList;
+
+namespace UE
+{
+namespace Insights
+{
+
+class FFilterConfigurator;
+class FTable;
+class FTableColumn;
+class FTreeNodeGrouping;
+class ITableCellValueSorter;
+
+////////////////////////////////////////////////////////////////////////////////////////////////////
+
+/** The filter collection - used for updating the list of tree nodes. */
+typedef TFilterCollection<const FTableTreeNodePtr&> FTableTreeNodeFilterCollection;
+
+/** The text based filter - used for updating the list of tree nodes. */
+typedef TTextFilter<const FTableTreeNodePtr&> FTableTreeNodeTextFilter;
+
+////////////////////////////////////////////////////////////////////////////////////////////////////
+
+enum class EAsyncOperationType : uint32
+{
+	FilteringOp = 1,
+	SortingOp = 1 << 1,
+	GroupingOp = 1 << 2,
+};
+
+ENUM_CLASS_FLAGS(EAsyncOperationType);
+
+struct FTableColumnConfig
+{
+	FName ColumnId;
+	bool bIsVisible;
+	float Width;
+};
+
+class ITableTreeViewPreset
+{
+public:
+	virtual FText GetName() const = 0;
+	virtual FText GetToolTip() const = 0;
+	virtual FName GetSortColumn() const = 0;
+	virtual EColumnSortMode::Type GetSortMode() const = 0;
+	virtual void SetCurrentGroupings(const TArray<TSharedPtr<FTreeNodeGrouping>>& InAvailableGroupings, TArray<TSharedPtr<FTreeNodeGrouping>>& InOutCurrentGroupings) const = 0;
+	virtual void GetColumnConfigSet(TArray<FTableColumnConfig>& InOutConfigSet) const = 0;
+};
+
+class FTableTaskCancellationToken
+{
+public:
+	FTableTaskCancellationToken()
+		: bCancel(false)
+	{}
+
+	bool ShouldCancel() { return bCancel.load(); }
+	void Cancel() { bCancel.store(true); }
+
+private:
+	std::atomic<bool> bCancel;
+};
+
+struct FTableTaskInfo
+{
+	FGraphEventRef Event;
+	TSharedPtr< FTableTaskCancellationToken> CancellationToken;
+};
+
+////////////////////////////////////////////////////////////////////////////////////////////////////
+/**
+ * A custom widget used to display the list of tree nodes.
+ */
+class STableTreeView : public SCompoundWidget, public IAsyncOperationStatusProvider
+{
+	friend class FTableTreeViewFilterAsyncTask;
+	friend class FTableTreeViewSortAsyncTask;
+	friend class FTableTreeViewGroupAsyncTask;
+	friend class FTableTreeViewAsyncCompleteTask;
+	friend class FSearchForItemToSelectTask;
+	friend class FSelectNodeByTableRowIndexTask;
+
+public:
+	/** Default constructor. */
+	STableTreeView();
+
+	/** Virtual destructor. */
+	virtual ~STableTreeView();
+
+	SLATE_BEGIN_ARGS(STableTreeView) {}
+	SLATE_END_ARGS()
+
+	/**
+	 * Construct this widget
+	 * @param InArgs - The declaration data for this widget
+	 */
+	void Construct(const FArguments& InArgs, TSharedPtr<FTable> InTablePtr);
+
+	TSharedPtr<FTable>& GetTable() { return Table; }
+	const TSharedPtr<FTable>& GetTable() const { return Table; }
+
+	virtual void Reset();
+
+	void RebuildColumns();
+
+	/**
+	 * Ticks this widget.  Override in derived classes, but always call the parent implementation.
+	 *
+	 * @param  AllottedGeometry The space allotted for this widget
+	 * @param  InCurrentTime  Current absolute real time
+	 * @param  InDeltaTime  Real time passed since last tick
+	 */
+	virtual void Tick(const FGeometry& AllottedGeometry, const double InCurrentTime, const float InDeltaTime) override;
+
+	/**
+	 * Rebuilds the tree (if necessary).
+	 * @param bResync - If true, it forces a resync even if the list did not changed since last sync.
+	 */
+	virtual void RebuildTree(bool bResync);
+
+	FTableTreeNodePtr GetNodeByTableRowIndex(int32 RowIndex) const;
+	void SelectNodeByTableRowIndex(int32 RowIndex);
+	bool IsRunningAsyncUpdate() { return bIsUpdateRunning;  }
+
+	void OnClose();
+
+	virtual FReply OnKeyDown(const FGeometry& MyGeometry, const FKeyEvent& InKeyEvent) override;
+
+	////////////////////////////////////////////////////////////////////////////////////////////////////
+	// IAsyncOperationStatusProvider implementation
+
+	virtual bool IsRunning() const override { return bIsUpdateRunning; }
+
+	virtual double GetAllOperationsDuration() override;
+	virtual double GetCurrentOperationDuration() override { return 0.0; }
+	virtual uint32 GetOperationCount() const override { return 1; }
+	virtual FText GetCurrentOperationName() const override;
+
+	////////////////////////////////////////////////////////////////////////////////////////////////////
+
+	/** Set a log listing name to be used for any errors or warnings. Must be preregistered by the caller with the MessageLog module. */
+	void SetLogListingName(const FName& InLogListingName) { LogListingName = InLogListingName; }
+	const FName& GetLogListingName() { return LogListingName; }
+
+protected:
+	void InitCommandList();
+
+	virtual void ConstructWidget(TSharedPtr<FTable> InTablePtr);
+	virtual TSharedRef<SWidget> ConstructSearchBox();
+	virtual TSharedRef<SWidget> ConstructAdvancedFiltersButton();
+	virtual TSharedRef<SWidget> ConstructHierarchyBreadcrumbTrail();
+	virtual TSharedPtr<SWidget> ConstructToolbar() { return nullptr; }
+	virtual TSharedPtr<SWidget> ConstructFooter() { return nullptr; }
+	virtual void ConstructHeaderArea(TSharedRef<SVerticalBox> InWidgetContent);
+	virtual void ConstructFooterArea(TSharedRef<SVerticalBox> InWidgetContent);
+
+	void UpdateTree();
+
+	/**
+	 * Populates OutSearchStrings with the strings that should be used in searching.
+	 *
+	 * @param GroupOrStatNodePtr - the group and stat node to get a text description from.
+	 * @param OutSearchStrings   - an array of strings to use in searching.
+	 *
+	 */
+	static void HandleItemToStringArray(const FTableTreeNodePtr& GroupOrStatNodePtr, TArray<FString>& OutSearchStrings);
+
+	////////////////////////////////////////////////////////////////////////////////////////////////////
+	// Tree View - Context Menu
+
+	TSharedPtr<SWidget> TreeView_GetMenuContent();
+	void TreeView_BuildSortByMenu(FMenuBuilder& MenuBuilder);
+	void TreeView_BuildViewColumnMenu(FMenuBuilder& MenuBuilder);
+	void TreeView_BuildExportMenu(FMenuBuilder& MenuBuilder);
+
+	bool ContextMenu_CopySelectedToClipboard_CanExecute() const;
+	void ContextMenu_CopySelectedToClipboard_Execute();
+	bool ContextMenu_CopyColumnToClipboard_CanExecute() const;
+	void ContextMenu_CopyColumnToClipboard_Execute();
+	bool ContextMenu_CopyColumnTooltipToClipboard_CanExecute() const;
+	void ContextMenu_CopyColumnTooltipToClipboard_Execute();
+	bool ContextMenu_ExpandSubtree_CanExecute() const;
+	void ContextMenu_ExpandSubtree_Execute();
+	bool ContextMenu_ExpandCriticalPath_CanExecute() const;
+	void ContextMenu_ExpandCriticalPath_Execute();
+	bool ContextMenu_CollapseSubtree_CanExecute() const;
+	void ContextMenu_CollapseSubtree_Execute();
+	bool ContextMenu_ExportToFile_CanExecute() const;
+	void ContextMenu_ExportToFile_Execute(bool bInExportCollapsed, bool InExportLeafs);
+
+	////////////////////////////////////////////////////////////////////////////////////////////////////
+	// Tree View - Columns' Header
+
+	void InitializeAndShowHeaderColumns();
+
+	FText GetColumnHeaderText(const FName ColumnId) const;
+
+	TSharedRef<SWidget> TreeViewHeaderRow_GenerateColumnMenu(const FTableColumn& Column);
+
+	////////////////////////////////////////////////////////////////////////////////////////////////////
+	// Tree View - Misc
+
+	void TreeView_Refresh();
+
+	/**
+	 * Called by STreeView to retrieves the children for the specified parent item.
+	 * @param InParent    - The parent node to retrieve the children from.
+	 * @param OutChildren - List of children for the parent node.
+	 */
+	void TreeView_OnGetChildren(FTableTreeNodePtr InParent, TArray<FTableTreeNodePtr>& OutChildren);
+
+	/** Called by STreeView when selection has changed. */
+	void TreeView_OnSelectionChanged(FTableTreeNodePtr SelectedItem, ESelectInfo::Type SelectInfo);
+
+	/** Called by STreeView when a tree item is double clicked. */
+	virtual void TreeView_OnMouseButtonDoubleClick(FTableTreeNodePtr TreeNode);
+
+	////////////////////////////////////////////////////////////////////////////////////////////////////
+	// Tree View - Table Row
+
+	/** Called by STreeView to generate a table row for the specified item. */
+	TSharedRef<ITableRow> TreeView_OnGenerateRow(FTableTreeNodePtr TreeNode, const TSharedRef<STableViewBase>& OwnerTable);
+
+	bool TableRow_ShouldBeEnabled(FTableTreeNodePtr NodePtr) const;
+
+	void TableRow_SetHoveredCell(TSharedPtr<FTable> TablePtr, TSharedPtr<FTableColumn> ColumnPtr, FTableTreeNodePtr NodePtr);
+	EHorizontalAlignment TableRow_GetColumnOutlineHAlignment(const FName ColumnId) const;
+
+	FText TableRow_GetHighlightText() const;
+	FName TableRow_GetHighlightedNodeName() const;
+
+	////////////////////////////////////////////////////////////////////////////////////////////////////
+	// Filtering
+
+	void OnFilteringChanged();
+
+	/** Populates the group and stat tree with items based on the current data. */
+	void ApplyFiltering();
+
+	bool ApplyAdvancedFiltersForNode(FTableTreeNodePtr NodePtr);
+
+	bool ApplyHierarchicalFilterForNode(FTableTreeNodePtr NodePtr, bool bFilterIsEmpty);
+
+	/** Set all the nodes belonging to a subtree as visible. Returns true if the caller node should be expanded. */
+	bool MakeSubtreeVisible(FTableTreeNodePtr NodePtr, bool bFilterIsEmpty);
+
+	bool SearchBox_IsEnabled() const;
+	void SearchBox_OnTextChanged(const FText& InFilterText);
+
+	FText SearchBox_GetTooltipText() const;
+
+	////////////////////////////////////////////////////////////////////////////////////////////////////
+	// Grouping
+
+	void CreateGroupings();
+	virtual void InternalCreateGroupings();
+
+	void OnGroupingChanged();
+	void ApplyGrouping();
+
+	void CreateGroups(const TArray<TSharedPtr<FTreeNodeGrouping>>& Groupings);
+	void GroupNodesRec(const TArray<FTableTreeNodePtr>& Nodes, FTableTreeNode& ParentGroup, int32 GroupingDepth, const TArray<TSharedPtr<FTreeNodeGrouping>>& Groupings);
+
+	void RebuildGroupingCrumbs();
+	void OnGroupingCrumbClicked(const TSharedPtr<FTreeNodeGrouping>& InEntry);
+	void BuildGroupingSubMenu_Change(FMenuBuilder& MenuBuilder, const TSharedPtr<FTreeNodeGrouping> CrumbGrouping);
+	void BuildGroupingSubMenu_Add(FMenuBuilder& MenuBuilder, const TSharedPtr<FTreeNodeGrouping> CrumbGrouping);
+	TSharedRef<SWidget> GetGroupingCrumbMenuContent(const TSharedPtr<FTreeNodeGrouping>& CrumbGrouping);
+
+	void PreChangeGroupings();
+	void PostChangeGroupings();
+	int32 GetGroupingDepth(const TSharedPtr<FTreeNodeGrouping>& Grouping) const;
+
+	void GroupingCrumbMenu_Reset_Execute();
+	void GroupingCrumbMenu_Remove_Execute(const TSharedPtr<FTreeNodeGrouping> Grouping);
+	void GroupingCrumbMenu_MoveLeft_Execute(const TSharedPtr<FTreeNodeGrouping> Grouping);
+	void GroupingCrumbMenu_MoveRight_Execute(const TSharedPtr<FTreeNodeGrouping> Grouping);
+	void GroupingCrumbMenu_Change_Execute(const TSharedPtr<FTreeNodeGrouping> OldGrouping, const TSharedPtr<FTreeNodeGrouping> NewGrouping);
+	bool GroupingCrumbMenu_Change_CanExecute(const TSharedPtr<FTreeNodeGrouping> OldGrouping, const TSharedPtr<FTreeNodeGrouping> NewGrouping) const;
+	void GroupingCrumbMenu_Add_Execute(const TSharedPtr<FTreeNodeGrouping> Grouping, const TSharedPtr<FTreeNodeGrouping> AfterGrouping);
+	bool GroupingCrumbMenu_Add_CanExecute(const TSharedPtr<FTreeNodeGrouping> Grouping, const TSharedPtr<FTreeNodeGrouping> AfterGrouping) const;
+
+	////////////////////////////////////////////////////////////////////////////////////////////////////
+	// Aggregation
+
+	void UpdateAggregatedValues(FTableTreeNode& GroupNode);
+
+	template<typename T>
+	static void UpdateAggregationRec(FTableColumn& Column, FTableTreeNode& GroupNode, T InitialAggregatedValue, bool bSetInitialValue, TFunctionRef<T(T, const FTableCellValue&)> ValueGetterFunc)
+	{
+		T AggregatedValue = InitialAggregatedValue;
+
+		for (FBaseTreeNodePtr NodePtr : GroupNode.GetChildren())
+		{
+			if (NodePtr->IsFiltered())
+			{
+				continue;
+			}
+
+			if (!NodePtr->IsGroup())
+			{
+				const TOptional<FTableCellValue> NodeValue = Column.GetValue(*NodePtr);
+				if (NodeValue.IsSet())
+				{
+					AggregatedValue = ValueGetterFunc(AggregatedValue, NodeValue.GetValue());
+				}
+			}
+			else
+			{
+				FTableTreeNode& TableNode = *(FTableTreeNode*)NodePtr.Get();
+				TableNode.ResetAggregatedValues(Column.GetId());
+				UpdateAggregationRec(Column, TableNode, InitialAggregatedValue, bSetInitialValue, ValueGetterFunc);
+				if (TableNode.HasAggregatedValue(Column.GetId()))
+				{
+					AggregatedValue = ValueGetterFunc(AggregatedValue, TableNode.GetAggregatedValue(Column.GetId()));
+				}
+			}
+		}
+
+		if (bSetInitialValue || InitialAggregatedValue != AggregatedValue)
+		{
+			GroupNode.AddAggregatedValue(Column.GetId(), FTableCellValue(AggregatedValue));
+		}
+	}
+
+	static void UpdateCStringSameValueAggregationRec(FTableColumn& Column, FTableTreeNode& GroupNode);
+
+	////////////////////////////////////////////////////////////////////////////////////////////////////
+	// Sorting
+
+	static const EColumnSortMode::Type GetDefaultColumnSortMode();
+	static const FName GetDefaultColumnBeingSorted();
+
+	void OnSortingChanged();
+	void CreateSortings();
+
+	void UpdateCurrentSortingByColumn();
+
+	void ApplySorting();
+	void SortTreeNodes(ITableCellValueSorter* InSorter, EColumnSortMode::Type InColumnSortMode);
+	void SortTreeNodesRec(FTableTreeNode& GroupNode, const ITableCellValueSorter& Sorter, EColumnSortMode::Type InColumnSortMode);
+
+	EColumnSortMode::Type GetSortModeForColumn(const FName ColumnId) const;
+	void SetSortModeForColumn(const FName& ColumnId, EColumnSortMode::Type SortMode);
+	void OnSortModeChanged(const EColumnSortPriority::Type SortPriority, const FName& ColumnId, const EColumnSortMode::Type SortMode);
+
+	////////////////////////////////////////////////////////////////////////////////////////////////////
+	// Sorting actions
+
+	// SortMode (HeaderMenu)
+	bool HeaderMenu_SortMode_IsChecked(const FName ColumnId, const EColumnSortMode::Type InSortMode);
+	bool HeaderMenu_SortMode_CanExecute(const FName ColumnId, const EColumnSortMode::Type InSortMode) const;
+	void HeaderMenu_SortMode_Execute(const FName ColumnId, const EColumnSortMode::Type InSortMode);
+
+	// SortMode (ContextMenu)
+	bool ContextMenu_SortMode_IsChecked(const EColumnSortMode::Type InSortMode);
+	bool ContextMenu_SortMode_CanExecute(const EColumnSortMode::Type InSortMode) const;
+	void ContextMenu_SortMode_Execute(const EColumnSortMode::Type InSortMode);
+
+	// SortByColumn (ContextMenu)
+	bool ContextMenu_SortByColumn_IsChecked(const FName ColumnId);
+	bool ContextMenu_SortByColumn_CanExecute(const FName ColumnId) const;
+	void ContextMenu_SortByColumn_Execute(const FName ColumnId);
+
+	////////////////////////////////////////////////////////////////////////////////////////////////////
+	// Column visibility actions
+
+	// ShowColumn
+	bool CanShowColumn(const FName ColumnId) const;
+	void ShowColumn(const FName ColumnId);
+	void ShowColumn(FTableColumn& Column);
+
+	// HideColumn
+	bool CanHideColumn(const FName ColumnId) const;
+	void HideColumn(const FName ColumnId);
+	void HideColumn(FTableColumn& Column);
+
+	// ToggleColumnVisibility
+	bool IsColumnVisible(const FName ColumnId);
+	bool CanToggleColumnVisibility(const FName ColumnId) const;
+	void ToggleColumnVisibility(const FName ColumnId);
+
+	// ShowAllColumns (ContextMenu)
+	bool ContextMenu_ShowAllColumns_CanExecute() const;
+	void ContextMenu_ShowAllColumns_Execute();
+
+	// ResetColumns (ContextMenu)
+	bool ContextMenu_ResetColumns_CanExecute() const;
+	void ContextMenu_ResetColumns_Execute();
+
+	// HideAllColumns (ContextMenu)
+	bool ContextMenu_HideAllColumns_CanExecute() const;
+	void ContextMenu_HideAllColumns_Execute();
+
+	//Async
+	virtual void OnPreAsyncUpdate();
+	virtual void OnPostAsyncUpdate();
+
+	FGraphEventRef StartSortTreeNodesTask(FGraphEventRef Prerequisite = nullptr);
+	FGraphEventRef StartCreateGroupsTask(FGraphEventRef Prerequisite = nullptr);
+	FGraphEventRef StartApplyFiltersTask(FGraphEventRef Prerequisite = nullptr);
+
+	void AddInProgressAsyncOperation(EAsyncOperationType InType) { EnumAddFlags(InProgressAsyncOperations, InType); }
+	bool HasInProgressAsyncOperation(EAsyncOperationType InType) const { return EnumHasAnyFlags(InProgressAsyncOperations, InType); }
+	void ClearInProgressAsyncOperations() { InProgressAsyncOperations = static_cast<EAsyncOperationType>(0); }
+
+	void StartPendingAsyncOperations();
+
+	void CancelCurrentAsyncOp();
+
+	FReply OnAdvancedFiltersClicked();
+	bool AdvancedFilters_ShouldBeEnabled() const;
+	FText AdvancedFilters_GetTooltipText() const;
+	bool FilterConfigurator_HasFilters() const;
+	void OnAdvancedFiltersChangesCommited();
+	bool ApplyAdvancedFilters(const FTableTreeNodePtr& NodePtr);
+	bool virtual ApplyCustomAdvancedFilters(const FTableTreeNodePtr& NodePtr) { return true; };
+	virtual void AddCustomAdvancedFilters() {}
+
+	////////////////////////////////////////////////////////////////////////////////////////////////////
+
+	void SetExpandValueForChildGroups(FBaseTreeNode* InRoot, int32 InMaxExpandedNodes, int32 MaxDepthToExpand, bool InValue);
+	void CountNumNodesPerDepthRec(FBaseTreeNode* InRoot, TArray<int32>& InOutNumNodesPerDepth, int32 InDepth, int32 InMaxDepth, int InMaxNodes) const;
+	void SetExpandValueForChildGroupsRec(FBaseTreeNode* InRoot, int32 InDepth, int32 InMaxDepth, bool InValue);
+
+	virtual void ExtendMenu(FMenuBuilder& Menu) {}
+
+	typedef TFunctionRef<void(TArray<FBaseTreeNodePtr>& InNodes)> WriteToFileCallback;
+	void ExportToFileRec(const FBaseTreeNodePtr& InGroupNode, TArray<FBaseTreeNodePtr>& InNodes, bool bInExportCollapsed, bool InExportLeafs, WriteToFileCallback Callback);
+
+	FText GetTreeViewBannerText() const { return TreeViewBannerText; }
+
+	virtual void InitAvailableViewPresets() {};
+	const TArray<TSharedRef<ITableTreeViewPreset>>* GetAvailableViewPresets() const { return &AvailableViewPresets; }
+	FReply OnApplyViewPreset(const ITableTreeViewPreset* InPreset);
+	void ApplyViewPreset(const ITableTreeViewPreset& InPreset);
+	void ApplyColumnConfig(const TArrayView<FTableColumnConfig>& InTableConfig);
+	void ViewPreset_OnSelectionChanged(TSharedPtr<ITableTreeViewPreset> InPreset, ESelectInfo::Type SelectInfo);
+	TSharedRef<SWidget> ViewPreset_OnGenerateWidget(TSharedRef<ITableTreeViewPreset> InPreset);
+	FText ViewPreset_GetSelectedText() const;
+	FText ViewPreset_GetSelectedToolTipText() const;
+
+	virtual void SearchForItem(TSharedPtr<FTableTaskCancellationToken> CancellationToken) {};
+
+	// Table data tasks should be tasks that operate read only operations on the data from the Table
+	// They should not operate on the tree nodes because they will run concurrently with the populated table UI.
+	template<typename T, typename... TArgs>
+	TSharedPtr<FTableTaskInfo> StartTableDataTask(TArgs&&... Args)
+	{
+		TSharedPtr<FTableTaskInfo> Info = MakeShared<FTableTaskInfo>();
+		Info->CancellationToken = MakeShared<FTableTaskCancellationToken>();
+		Info->Event = TGraphTask<T>::CreateTask().ConstructAndDispatchWhenReady(Info->CancellationToken, Forward<TArgs>(Args)...);
+		DataTaskInfos.Add(Info);
+
+		return Info;
+	}
+
+	void StopAllTableDataTasks(bool bWait = true);
+
+protected:
+	/** Table view model. */
+	TSharedPtr<FTable> Table;
+	//////////////////////////////////////////////////
+	// Tree View, Columns
+
+	/** The child STreeView widget. */
+	TSharedPtr<STreeView<FTableTreeNodePtr>> TreeView;
+
+	/** Holds the tree view header row widget which display all columns in the tree view. */
+	TSharedPtr<SHeaderRow> TreeViewHeaderRow;
+
+	/** External scrollbar used to synchronize tree view position. */
+	TSharedPtr<SScrollBar> ExternalScrollbar;
+
+	//////////////////////////////////////////////////
+	// Hovered Column, Hovered Tree Node
+
+	/** Name of the column currently being hovered by the mouse. */
+	FName HoveredColumnId;
+
+	/** A shared pointer to the tree node currently being hovered by the mouse. */
+	FTableTreeNodePtr HoveredNodePtr;
+
+	/** Name of the tree node that should be drawn as highlighted. */
+	FName HighlightedNodeName;
+
+	TSharedPtr<FUICommandList> CommandList;
+
+	//////////////////////////////////////////////////
+	// Tree Nodes
+
+	static const FName RootNodeName;
+
+	/** The root node of the tree. */
+	FTableTreeNodePtr Root;
+
+	/** Table (row) nodes. Each node corresponds to a table row. Index in this array corresponds to RowIndex in source table. */
+	TArray<FTableTreeNodePtr> TableTreeNodes;
+
+	/** A filtered array of group and nodes to be displayed in the tree widget. */
+	TArray<FTableTreeNodePtr> FilteredGroupNodes;
+
+	/** Currently expanded group nodes. */
+	TSet<FTableTreeNodePtr> ExpandedNodes;
+
+	/** If true, the expanded nodes have been saved before applying a text filter. */
+	bool bExpansionSaved;
+
+	//////////////////////////////////////////////////
+	// Search box and filters
+
+	/** The search box widget used to filter items displayed in the stats and groups tree. */
+	TSharedPtr<SSearchBox> SearchBox;
+
+	/** The text based filter. */
+	TSharedPtr<FTableTreeNodeTextFilter> TextFilter;
+
+	/** The filter collection. */
+	TSharedPtr<FTableTreeNodeFilterCollection> Filters;
+
+	//////////////////////////////////////////////////
+	// Grouping
+
+	TArray<TSharedPtr<FTreeNodeGrouping>> AvailableGroupings;
+
+	/** How we group the tree nodes? */
+	TArray<TSharedPtr<FTreeNodeGrouping>> CurrentGroupings;
+
+	TSharedPtr<SBreadcrumbTrail<TSharedPtr<FTreeNodeGrouping>>> GroupingBreadcrumbTrail;
+
+	//////////////////////////////////////////////////
+	// Sorting
+
+	/** All available sorters. */
+	TArray<TSharedPtr<ITableCellValueSorter>> AvailableSorters;
+
+	/** Current sorter. It is nullptr if sorting is disabled. */
+	TSharedPtr<ITableCellValueSorter> CurrentSorter;
+
+	/** Name of the column currently being sorted. Can be NAME_None if sorting is disabled (CurrentSorting == nullptr) or if a complex sorting is used (CurrentSorting != nullptr). */
+	FName ColumnBeingSorted;
+
+	/** How we sort the nodes? Ascending or Descending. */
+	EColumnSortMode::Type ColumnSortMode;
+
+	//////////////////////////////////////////////////
+	// Async
+
+	bool bRunInAsyncMode = false;
+	bool bIsUpdateRunning = false;
+	bool bIsCloseScheduled = false;
+
+	TArray<FTableTreeNodePtr> DummyGroupNodes;
+	FGraphEventRef InProgressAsyncOperationEvent;
+	FGraphEventRef AsyncCompleteTaskEvent;
+	EAsyncOperationType InProgressAsyncOperations = static_cast<EAsyncOperationType>(0);
+	TSharedPtr<class SAsyncOperationStatus> AsyncOperationStatus;
+	FStopwatch AsyncUpdateStopwatch;
+
+	TArray<TSharedPtr<FTreeNodeGrouping>> CurrentAsyncOpGroupings;
+	ITableCellValueSorter* CurrentAsyncOpSorter = nullptr;
+	EColumnSortMode::Type CurrentAsyncOpColumnSortMode;
+	TSharedPtr<FTableTreeNodeTextFilter> CurrentAsyncOpTextFilter;
+	FFilterConfigurator* CurrentAsyncOpFilterConfigurator = nullptr;
+	FAsyncOperationProgress AsyncOperationProgress;
+	FGraphEventRef DispatchEvent;
+	TArray<TSharedPtr<FTableTaskInfo>> DataTaskInfos;
+	TArray<FTableTreeNodePtr> NodesToExpand;
+
+	//////////////////////////////////////////////////
+
+	TSharedPtr<FFilterConfigurator> FilterConfigurator;
+	FDelegateHandle OnFilterChangesCommitedHandle;
+	FFilterContext Context;
+
+	double StatsStartTime;
+	double StatsEndTime;
+
+	FText TreeViewBannerText;
+
+	TArray<TSharedRef<ITableTreeViewPreset>> AvailableViewPresets;
+	TSharedPtr<ITableTreeViewPreset> SelectedViewPreset;
+	TSharedPtr<SComboBox<TSharedRef<ITableTreeViewPreset>>> PresetComboBox;
+
+	static constexpr int32 MAX_NUMBER_OF_NODES_TO_EXPAND = 1000 * 1000;
+	static constexpr int32 MAX_DEPTH_TO_EXPAND = 100;
+
+	//////////////////////////////////////////////////
+	// Logging
+
+	/** A log listing name to be used for any errors or warnings. */
+	FName LogListingName;
+};
+
+////////////////////////////////////////////////////////////////////////////////////////////////////
+
+class FTableTreeViewFilterAsyncTask
+{
+public:
+	FTableTreeViewFilterAsyncTask(STableTreeView* InPtr)
+	{
+		TableTreeViewPtr = InPtr;
+	}
+
+	FORCEINLINE TStatId GetStatId() const { RETURN_QUICK_DECLARE_CYCLE_STAT(FTableTreeViewFilterAsyncTask, STATGROUP_TaskGraphTasks); }
+	ENamedThreads::Type GetDesiredThread() { return ENamedThreads::Type::AnyThread; }
+	static ESubsequentsMode::Type GetSubsequentsMode() { return ESubsequentsMode::TrackSubsequents; }
+
+	void DoTask(ENamedThreads::Type CurrentThread, const FGraphEventRef& MyCompletionGraphEvent)
+	{
+		if (TableTreeViewPtr)
+		{
+			TableTreeViewPtr->ApplyFiltering();
+		}
+	}
+
+private:
+	STableTreeView* TableTreeViewPtr = nullptr;
+};
+
+////////////////////////////////////////////////////////////////////////////////////////////////////
+
+class FTableTreeViewSortAsyncTask
+{
+public:
+	FTableTreeViewSortAsyncTask(STableTreeView* InPtr, ITableCellValueSorter* InSorter, EColumnSortMode::Type InColumnSortMode)
+	{
+		TableTreeViewPtr = InPtr;
+		Sorter = InSorter;
+		ColumnSortMode = InColumnSortMode;
+	}
+
+	FORCEINLINE TStatId GetStatId() const { RETURN_QUICK_DECLARE_CYCLE_STAT(FTableTreeViewSortAsyncTask, STATGROUP_TaskGraphTasks); }
+	ENamedThreads::Type GetDesiredThread() { return ENamedThreads::Type::AnyThread; }
+	static ESubsequentsMode::Type GetSubsequentsMode() { return ESubsequentsMode::TrackSubsequents; }
+
+	void DoTask(ENamedThreads::Type CurrentThread, const FGraphEventRef& MyCompletionGraphEvent)
+	{
+		if (TableTreeViewPtr)
+		{
+			TableTreeViewPtr->SortTreeNodes(Sorter, ColumnSortMode);
+		}
+	}
+
+private:
+	STableTreeView* TableTreeViewPtr;
+	ITableCellValueSorter* Sorter;
+	EColumnSortMode::Type ColumnSortMode;
+};
+
+////////////////////////////////////////////////////////////////////////////////////////////////////
+
+class FTableTreeViewGroupAsyncTask
+{
+public:
+	FTableTreeViewGroupAsyncTask(STableTreeView* InPtr, TArray<TSharedPtr<FTreeNodeGrouping>>* InGroupings)
+	{
+		TableTreeViewPtr = InPtr;
+		Groupings = InGroupings;
+	}
+
+	FORCEINLINE TStatId GetStatId() const { RETURN_QUICK_DECLARE_CYCLE_STAT(FTableTreeViewGroupAsyncTask, STATGROUP_TaskGraphTasks); }
+	ENamedThreads::Type GetDesiredThread() { return ENamedThreads::Type::AnyThread; }
+	static ESubsequentsMode::Type GetSubsequentsMode() { return ESubsequentsMode::TrackSubsequents; }
+
+	void DoTask(ENamedThreads::Type CurrentThread, const FGraphEventRef& MyCompletionGraphEvent)
+	{
+		if (TableTreeViewPtr)
+		{
+			TableTreeViewPtr->CreateGroups(*Groupings);
+		}
+	}
+
+private:
+	STableTreeView* TableTreeViewPtr = nullptr;
+	TArray<TSharedPtr<FTreeNodeGrouping>>* Groupings;
+};
+
+////////////////////////////////////////////////////////////////////////////////////////////////////
+
+class FTableTreeViewAsyncCompleteTask
+{
+public:
+	FTableTreeViewAsyncCompleteTask(TSharedPtr<STableTreeView> InPtr)
+	{
+		TableTreeViewPtr = InPtr;
+	}
+
+	FORCEINLINE TStatId GetStatId() const { RETURN_QUICK_DECLARE_CYCLE_STAT(FTableTreeViewAsyncCompleteTask, STATGROUP_TaskGraphTasks); }
+	ENamedThreads::Type GetDesiredThread() { return ENamedThreads::Type::GameThread; }
+	static ESubsequentsMode::Type GetSubsequentsMode() { return ESubsequentsMode::TrackSubsequents; }
+
+	void DoTask(ENamedThreads::Type CurrentThread, const FGraphEventRef& MyCompletionGraphEvent)
+	{
+		// The role of this task is to keep the STableTreeView object alive until the task and it's prerequisits are completed and to destroy it on the game thread.
+		FGraphEventRef Event = TableTreeViewPtr->AsyncCompleteTaskEvent;
+		if (TableTreeViewPtr.IsValid())
+		{
+			TableTreeViewPtr.Reset();
+		}
+		//TODO: callback for FInsightsManager::Get()->RemoveInProgressAsyncOp(Event);
+	}
+
+private:
+	TSharedPtr<STableTreeView> TableTreeViewPtr;
+};
+
+////////////////////////////////////////////////////////////////////////////////////////////////////
+
+class FSearchForItemToSelectTask
+{
+public:
+	FSearchForItemToSelectTask(TSharedPtr<FTableTaskCancellationToken> InToken, TSharedPtr<STableTreeView> InPtr)
+		: CancellationToken(InToken)
+		, TableTreeViewPtr(InPtr)
+	{}
+
+	FORCEINLINE TStatId GetStatId() const { RETURN_QUICK_DECLARE_CYCLE_STAT(FSearchForItemToSelectTask, STATGROUP_TaskGraphTasks); }
+	ENamedThreads::Type GetDesiredThread() { return ENamedThreads::Type::AnyThread; }
+	static ESubsequentsMode::Type GetSubsequentsMode() { return ESubsequentsMode::TrackSubsequents; }
+
+	void DoTask(ENamedThreads::Type CurrentThread, const FGraphEventRef& MyCompletionGraphEvent)
+	{
+		TableTreeViewPtr->SearchForItem(CancellationToken);
+	}
+
+private:
+	TSharedPtr<FTableTaskCancellationToken> CancellationToken;
+	TSharedPtr<STableTreeView> TableTreeViewPtr;
+};
+
+////////////////////////////////////////////////////////////////////////////////////////////////////
+
+class FSelectNodeByTableRowIndexTask
+{
+public:
+	FSelectNodeByTableRowIndexTask(TSharedPtr<FTableTaskCancellationToken> InToken, TSharedPtr<STableTreeView> InPtr, uint32 InRowIndex)
+		: CancellationToken(InToken)
+		, TableTreeViewPtr(InPtr)
+		, RowIndex(InRowIndex)
+		{}
+
+	FORCEINLINE TStatId GetStatId() const { RETURN_QUICK_DECLARE_CYCLE_STAT(FSelectNodeByTableRowIndexTask, STATGROUP_TaskGraphTasks); }
+	ENamedThreads::Type GetDesiredThread() { return ENamedThreads::Type::GameThread; }
+	static ESubsequentsMode::Type GetSubsequentsMode() { return ESubsequentsMode::TrackSubsequents; }
+
+	void DoTask(ENamedThreads::Type CurrentThread, const FGraphEventRef& MyCompletionGraphEvent)
+	{
+		if (!CancellationToken->ShouldCancel())
+		{
+			TableTreeViewPtr->SelectNodeByTableRowIndex(RowIndex);
+		}
+	}
+
+private:
+	TSharedPtr< FTableTaskCancellationToken> CancellationToken;
+	TSharedPtr<STableTreeView> TableTreeViewPtr;
+	uint32 RowIndex;
+};
+
+////////////////////////////////////////////////////////////////////////////////////////////////////
+
+} // namespace Insights
+} // namespace UE
diff --git a/Engine/Plugins/Editor/AssetManagerEditor/Source/AssetManagerEditor/Private/Insights/Table/Widgets/STableTreeViewCell.cpp b/Engine/Plugins/Editor/AssetManagerEditor/Source/AssetManagerEditor/Private/Insights/Table/Widgets/STableTreeViewCell.cpp
new file mode 100644
--- /dev/null
+++ b/Engine/Plugins/Editor/AssetManagerEditor/Source/AssetManagerEditor/Private/Insights/Table/Widgets/STableTreeViewCell.cpp
@@ -0,0 +1,184 @@
+// Copyright Epic Games, Inc. All Rights Reserved.
+
+#include "STableTreeViewCell.h"
+
+#include "SlateOptMacros.h"
+#include "Widgets/Images/SImage.h"
+#include "Widgets/Input/SButton.h"
+#include "Widgets/SBoxPanel.h"
+#include "Widgets/SToolTip.h"
+#include "Widgets/Text/STextBlock.h"
+#include "Widgets/Views/SExpanderArrow.h"
+
+// Insights
+#include "Insights/Common/InsightsStyle.h"
+#include "Insights/Table/ViewModels/Table.h"
+#include "Insights/Table/ViewModels/TableCellValueFormatter.h"
+#include "Insights/Table/ViewModels/TableColumn.h"
+#include "Insights/Table/Widgets/STableTreeViewRow.h"
+
+#define LOCTEXT_NAMESPACE "STableTreeView"
+
+namespace UE
+{
+namespace Insights
+{
+
+////////////////////////////////////////////////////////////////////////////////////////////////////
+
+BEGIN_SLATE_FUNCTION_BUILD_OPTIMIZATION
+
+void STableTreeViewCell::Construct(const FArguments& InArgs, const TSharedRef<ITableRow>& TableRow)
+{
+	TablePtr = InArgs._TablePtr;
+	ColumnPtr = InArgs._ColumnPtr;
+	TableTreeNodePtr = InArgs._TableTreeNodePtr;
+
+	ensure(TablePtr.IsValid());
+	ensure(ColumnPtr.IsValid());
+	ensure(TableTreeNodePtr.IsValid());
+
+	SetHoveredCellDelegate = InArgs._OnSetHoveredCell;
+
+	ChildSlot
+	[
+		GenerateWidgetForColumn(InArgs, TableRow)
+	];
+}
+
+////////////////////////////////////////////////////////////////////////////////////////////////////
+
+TSharedRef<SWidget> STableTreeViewCell::GenerateWidgetForColumn(const FArguments& InArgs, const TSharedRef<ITableRow>& TableRow)
+{
+	if (InArgs._IsNameColumn)
+	{
+		return GenerateWidgetForNameColumn(InArgs, TableRow);
+	}
+	else
+	{
+		return GenerateWidgetForTableColumn(InArgs, TableRow);
+	}
+}
+
+////////////////////////////////////////////////////////////////////////////////////////////////////
+
+TSharedRef<SWidget> STableTreeViewCell::GenerateWidgetForNameColumn(const FArguments& InArgs, const TSharedRef<ITableRow>& TableRow)
+{
+	return
+		SNew(SHorizontalBox)
+
+		+ SHorizontalBox::Slot()
+		.AutoWidth()
+		.HAlign(HAlign_Right)
+		.VAlign(VAlign_Center)
+		[
+			SNew(SExpanderArrow, TableRow)
+		]
+
+		// Info icon + tooltip
+		+ SHorizontalBox::Slot()
+		.AutoWidth()
+		.HAlign(HAlign_Center)
+		.VAlign(VAlign_Center)
+		[
+			SNew(SImage)
+			.Visibility(this, &STableTreeViewCell::GetHintIconVisibility)
+			.Image(FInsightsStyle::GetBrush("Icons.Hint.TreeItem"))
+			.ToolTip(GetRowToolTip(TableRow))
+		]
+
+		// Name
+		+ SHorizontalBox::Slot()
+		.AutoWidth()
+		.VAlign(VAlign_Center)
+		.HAlign(HAlign_Left)
+		.Padding(FMargin(2.0f, 0.0f))
+		[
+			SNew(STextBlock)
+			.Text(this, &STableTreeViewCell::GetDisplayName)
+			.HighlightText(InArgs._HighlightText)
+			.TextStyle(FInsightsStyle::Get(), TEXT("TreeTable.Tooltip"))
+			.ColorAndOpacity(this, &STableTreeViewCell::GetColorAndOpacity)
+			.ShadowColorAndOpacity(this, &STableTreeViewCell::GetShadowColorAndOpacity)
+		]
+
+		// Name Suffix
+		+ SHorizontalBox::Slot()
+		.AutoWidth()
+		.VAlign(VAlign_Center)
+		.HAlign(HAlign_Left)
+		.Padding(FMargin(2.0f, 0.0f))
+		[
+			SNew(STextBlock)
+			.Visibility(this, &STableTreeViewCell::HasExtraDisplayName)
+			.Text(this, &STableTreeViewCell::GetExtraDisplayName)
+			.TextStyle(FInsightsStyle::Get(), TEXT("TreeTable.Tooltip"))
+			.ColorAndOpacity(this, &STableTreeViewCell::GetExtraColorAndOpacity)
+			.ShadowColorAndOpacity(this, &STableTreeViewCell::GetShadowColorAndOpacity)
+		]
+	;
+}
+
+////////////////////////////////////////////////////////////////////////////////////////////////////
+
+TSharedPtr<IToolTip> STableTreeViewCell::GetRowToolTip(const TSharedRef<ITableRow>& TableRow) const
+{
+	TSharedRef<STableTreeViewRow> Row = StaticCastSharedRef<STableTreeViewRow, ITableRow>(TableRow);
+	return Row->GetRowToolTip();
+}
+
+////////////////////////////////////////////////////////////////////////////////////////////////////
+
+FText STableTreeViewCell::GetValueAsText() const
+{
+	return ColumnPtr->GetValueAsText(*TableTreeNodePtr);
+}
+
+////////////////////////////////////////////////////////////////////////////////////////////////////
+
+TSharedRef<SWidget> STableTreeViewCell::GenerateWidgetForTableColumn(const FArguments& InArgs, const TSharedRef<ITableRow>& TableRow)
+{
+	TSharedRef<STextBlock> TextBox = SNew(STextBlock)
+		.TextStyle(FInsightsStyle::Get(), TEXT("TreeTable.Tooltip"))
+		.ColorAndOpacity(this, &STableTreeViewCell::GetStatsColorAndOpacity)
+		.ShadowColorAndOpacity(this, &STableTreeViewCell::GetShadowColorAndOpacity);
+
+	if (ColumnPtr->IsDynamic())
+	{
+		TextBox->SetText(TAttribute<FText>::Create(TAttribute<FText>::FGetter::CreateRaw(this, &STableTreeViewCell::GetValueAsText)));
+	}
+	else
+	{
+		FText CellText = ColumnPtr->GetValueAsText(*TableTreeNodePtr);
+		TextBox->SetText(CellText);
+	}
+
+	return
+		SNew(SHorizontalBox)
+		.ToolTip(GetTooltip())
+
+		+ SHorizontalBox::Slot()
+		.FillWidth(1.0f)
+		.VAlign(VAlign_Center)
+		.HAlign(ColumnPtr->GetHorizontalAlignment())
+		.Padding(FMargin(2.0f, 0.0f))
+		[
+			TextBox
+		];
+}
+
+END_SLATE_FUNCTION_BUILD_OPTIMIZATION
+
+////////////////////////////////////////////////////////////////////////////////////////////////////
+
+TSharedPtr<IToolTip> STableTreeViewCell::GetTooltip() const
+{
+	return ColumnPtr->GetValueFormatter()->GetCustomTooltip(*ColumnPtr, *TableTreeNodePtr);
+}
+
+////////////////////////////////////////////////////////////////////////////////////////////////////
+
+} // namespace Insights
+} // namespace UE
+
+#undef LOCTEXT_NAMESPACE
diff --git a/Engine/Plugins/Editor/AssetManagerEditor/Source/AssetManagerEditor/Private/Insights/Table/Widgets/STableTreeViewCell.h b/Engine/Plugins/Editor/AssetManagerEditor/Source/AssetManagerEditor/Private/Insights/Table/Widgets/STableTreeViewCell.h
new file mode 100644
--- /dev/null
+++ b/Engine/Plugins/Editor/AssetManagerEditor/Source/AssetManagerEditor/Private/Insights/Table/Widgets/STableTreeViewCell.h
@@ -0,0 +1,179 @@
+// Copyright Epic Games, Inc. All Rights Reserved.
+
+#pragma once
+
+#include "CoreMinimal.h"
+#include "Widgets/DeclarativeSyntaxSupport.h"
+#include "Widgets/SCompoundWidget.h"
+
+// Insights
+#include "Insights/Table/ViewModels/TableTreeNode.h"
+
+class ITableRow;
+class IToolTip;
+
+namespace UE
+{
+namespace Insights
+{
+
+class FTable;
+class FTableColumn;
+
+DECLARE_DELEGATE_ThreeParams(FSetHoveredTableTreeViewCell, TSharedPtr<FTable> /*TablePtr*/, TSharedPtr<FTableColumn> /*ColumnPtr*/, FTableTreeNodePtr /*TableTreeNodePtr*/);
+
+class STableTreeViewCell : public SCompoundWidget
+{
+public:
+	SLATE_BEGIN_ARGS(STableTreeViewCell) {}
+		SLATE_EVENT(FSetHoveredTableTreeViewCell, OnSetHoveredCell)
+		SLATE_ATTRIBUTE(FText, HighlightText)
+		SLATE_ARGUMENT(TSharedPtr<FTable>, TablePtr)
+		SLATE_ARGUMENT(TSharedPtr<FTableColumn>, ColumnPtr)
+		SLATE_ARGUMENT(FTableTreeNodePtr, TableTreeNodePtr)
+		SLATE_ARGUMENT(bool, IsNameColumn)
+	SLATE_END_ARGS()
+
+	void Construct(const FArguments& InArgs, const TSharedRef<ITableRow>& TableRow);
+
+protected:
+	TSharedRef<SWidget> GenerateWidgetForColumn(const FArguments& InArgs, const TSharedRef<ITableRow>& TableRow);
+	TSharedRef<SWidget> GenerateWidgetForNameColumn(const FArguments& InArgs, const TSharedRef<ITableRow>& TableRow);
+	TSharedRef<SWidget> GenerateWidgetForTableColumn(const FArguments& InArgs, const TSharedRef<ITableRow>& TableRow);
+
+	/**
+	 * The system will use this event to notify a widget that the cursor has entered it. This event is NOT bubbled.
+	 *
+	 * @param MyGeometry The Geometry of the widget receiving the event
+	 * @param MouseEvent Information about the input event
+	 */
+	virtual void OnMouseEnter(const FGeometry& MyGeometry, const FPointerEvent& MouseEvent) override
+	{
+		SCompoundWidget::OnMouseEnter(MyGeometry, MouseEvent);
+		SetHoveredCellDelegate.ExecuteIfBound(TablePtr, ColumnPtr, TableTreeNodePtr);
+	}
+
+	/**
+	 * The system will use this event to notify a widget that the cursor has left it. This event is NOT bubbled.
+	 *
+	 * @param MouseEvent Information about the input event
+	 */
+	virtual void OnMouseLeave(const FPointerEvent& MouseEvent) override
+	{
+		SCompoundWidget::OnMouseLeave(MouseEvent);
+		SetHoveredCellDelegate.ExecuteIfBound(nullptr, nullptr, nullptr);
+	}
+
+	/**
+	 * Called during drag and drop when the drag enters a widget.
+	 *
+	 * Enter/Leave events in slate are meant as lightweight notifications.
+	 * So we do not want to capture mouse or set focus in response to these.
+	 * However, OnDragEnter must also support external APIs (e.g. OLE Drag/Drop)
+	 * Those require that we let them know whether we can handle the content
+	 * being dragged OnDragEnter.
+	 *
+	 * The concession is to return a can_handled/cannot_handle
+	 * boolean rather than a full FReply.
+	 *
+	 * @param MyGeometry      The geometry of the widget receiving the event.
+	 * @param DragDropEvent   The drag and drop event.
+	 *
+	 * @return A reply that indicated whether the contents of the DragDropEvent can potentially be processed by this widget.
+	 */
+	virtual void OnDragEnter(const FGeometry& MyGeometry, const FDragDropEvent& DragDropEvent) override
+	{
+		SCompoundWidget::OnDragEnter(MyGeometry, DragDropEvent);
+		SetHoveredCellDelegate.ExecuteIfBound(TablePtr, ColumnPtr, TableTreeNodePtr);
+	}
+
+	/**
+	 * Called during drag and drop when the drag leaves a widget.
+	 *
+	 * @param DragDropEvent   The drag and drop event.
+	 */
+	virtual void OnDragLeave(const FDragDropEvent& DragDropEvent)  override
+	{
+		SCompoundWidget::OnDragLeave(DragDropEvent);
+		SetHoveredCellDelegate.ExecuteIfBound(nullptr, nullptr, nullptr);
+	}
+
+	EVisibility GetHintIconVisibility() const
+	{
+		return IsHovered() ? EVisibility::Visible : EVisibility::Hidden;
+	}
+
+	TSharedPtr<class IToolTip> GetRowToolTip(const TSharedRef<ITableRow>& TableRow) const;
+
+	FText GetDisplayName() const
+	{
+		return TableTreeNodePtr->GetDisplayName();
+	}
+
+	FText GetExtraDisplayName() const
+	{
+		return TableTreeNodePtr->GetExtraDisplayName();
+	}
+
+	EVisibility HasExtraDisplayName() const
+	{
+		return TableTreeNodePtr->HasExtraDisplayName() ? EVisibility::Visible : EVisibility::Collapsed;
+	}
+
+	FText GetValueAsText() const;
+
+	FSlateColor GetColorAndOpacity() const
+	{
+		const FLinearColor TextColor =
+			TableTreeNodePtr->IsFiltered() ?
+				FLinearColor(1.0f, 1.0f, 1.0f, 0.5f) :
+				FLinearColor(1.0f, 1.0f, 1.0f, 1.0f);
+		return TextColor;
+	}
+
+	FSlateColor GetExtraColorAndOpacity() const
+	{
+		const FLinearColor TextColor =
+			TableTreeNodePtr->IsFiltered() ?
+				FLinearColor(0.5f, 0.5f, 0.5f, 0.5f) :
+				FLinearColor(0.5f, 0.5f, 0.5f, 1.0f);
+		return TextColor;
+	}
+
+	FSlateColor GetStatsColorAndOpacity() const
+	{
+		const FLinearColor TextColor =
+			TableTreeNodePtr->IsGroup() ?
+				FLinearColor(0.7f, 0.7f, 0.7f, 1.0f) :
+			TableTreeNodePtr->IsFiltered() ?
+				FLinearColor(1.0f, 1.0f, 1.0f, 0.5f) :
+				FLinearColor(1.0f, 1.0f, 1.0f, 1.0f);
+		return TextColor;
+	}
+
+	FLinearColor GetShadowColorAndOpacity() const
+	{
+		const FLinearColor ShadowColor =
+			TableTreeNodePtr->IsFiltered() ?
+				FLinearColor(0.f, 0.f, 0.f, 0.25f) :
+				FLinearColor(0.0f, 0.0f, 0.0f, 0.5f);
+		return ShadowColor;
+	}
+
+	TSharedPtr<IToolTip> GetTooltip() const;
+
+protected:
+	/** A shared pointer to the table view model. */
+	TSharedPtr<FTable> TablePtr;
+
+	/** A shared pointer to the table column view model. */
+	TSharedPtr<FTableColumn> ColumnPtr;
+
+	/** A shared pointer to the tree node. */
+	FTableTreeNodePtr TableTreeNodePtr;
+
+	FSetHoveredTableTreeViewCell SetHoveredCellDelegate;
+};
+
+} // namespace Insights
+} // namespace UE
