AssetManagerEditor: Fixed tree view grouping by CString unique values.

diff --git a/Engine/Plugins/Editor/AssetManagerEditor/Source/AssetManagerEditor/Private/Insights/Table/ViewModels/TreeNodeGrouping.cpp b/Engine/Plugins/Editor/AssetManagerEditor/Source/AssetManagerEditor/Private/Insights/Table/ViewModels/TreeNodeGrouping.cpp
--- a/Engine/Plugins/Editor/AssetManagerEditor/Source/AssetManagerEditor/Private/Insights/Table/ViewModels/TreeNodeGrouping.cpp
+++ b/Engine/Plugins/Editor/AssetManagerEditor/Source/AssetManagerEditor/Private/Insights/Table/ViewModels/TreeNodeGrouping.cpp
@@ -154,12 +154,6 @@ template<> double TTreeNodeGroupingByUniqueValue<double>::GetValue(const FTableC
 
 FName FTreeNodeGroupingByUniqueValueCString::GetGroupName(const TCHAR* Value)
 {
-	if (Value == nullptr || *Value == TEXT('\0'))
-	{
-		static FName EmptyGroupName(TEXT("N/A"));
-		return EmptyGroupName;
-	}
-
 	FStringView StringView(Value);
 	if (StringView.Len() >= NAME_SIZE)
 	{
@@ -172,8 +166,10 @@ FName FTreeNodeGroupingByUniqueValueCString::GetGroupName(const TCHAR* Value)
 
 void FTreeNodeGroupingByUniqueValueCString::GroupNodes(const TArray<FTableTreeNodePtr>& Nodes, FTableTreeNode& ParentGroup, TWeakPtr<FTable> InParentTable, IAsyncOperationProgress& InAsyncOperationProgress) const
 {
-	TMap<const void*, FTableTreeNodePtr> GroupMap; // using void* as TMap<const TCHAR*, T> doesn't allow nullptr to be added
-	FTableTreeNodePtr UnsetGroupPtr = nullptr;
+	FTableTreeNodePtr UnsetGroupPtr = nullptr; // for unset FTableCellValue
+	FTableTreeNodePtr EmptyGroupPtr = nullptr; // for nullptr and empty strings
+	TMap<const TCHAR*, FTableTreeNodePtr> GroupMap; // for valid strings
+	TMap<FName, FTableTreeNodePtr> GroupNameMap; // for valid strings
 
 	ParentGroup.ClearChildren();
 
@@ -197,19 +193,40 @@ void FTreeNodeGroupingByUniqueValueCString::GroupNodes(const TArray<FTableTreeNo
 		if (CellValue.IsSet())
 		{
 			const TCHAR* Value = CellValue.GetValue().CString;
-
-			FTableTreeNodePtr* GroupPtrPtr = GroupMap.Find(Value);
-			if (!GroupPtrPtr)
+			if (Value != nullptr && Value[0] != TEXT('\0'))
 			{
-				const FName GroupName = GetGroupName(Value);
-				GroupPtr = MakeShared<FTableTreeNode>(GroupName, InParentTable);
-				GroupPtr->SetExpansion(false);
-				ParentGroup.AddChildAndSetGroupPtr(GroupPtr);
-				GroupMap.Add(Value, GroupPtr);
+				FTableTreeNodePtr* GroupPtrPtr = GroupMap.Find(Value);
+				if (!GroupPtrPtr)
+				{
+					const FName GroupName = GetGroupName(Value);
+					GroupPtrPtr = GroupNameMap.Find(Value);
+					if (!GroupPtrPtr)
+					{
+						GroupPtr = MakeShared<FTableTreeNode>(GroupName, InParentTable);
+						GroupPtr->SetExpansion(false);
+						ParentGroup.AddChildAndSetGroupPtr(GroupPtr);
+						GroupNameMap.Add(GroupName, GroupPtr);
+					}
+					else
+					{
+						GroupPtr = *GroupPtrPtr;
+					}
+					GroupMap.Add(Value, GroupPtr);
+				}
+				else
+				{
+					GroupPtr = *GroupPtrPtr;
+				}
 			}
 			else
 			{
-				GroupPtr = *GroupPtrPtr;
+				if (!EmptyGroupPtr)
+				{
+					EmptyGroupPtr = MakeShared<FTableTreeNode>(FName(TEXT("N/A")), InParentTable);
+					EmptyGroupPtr->SetExpansion(false);
+					ParentGroup.AddChildAndSetGroupPtr(EmptyGroupPtr);
+				}
+				GroupPtr = EmptyGroupPtr;
 			}
 		}
 		else
