AssetManagerEditor: Added initial Asset tree/table view based on Insights' TreeView. Code switched off by ASSET_TABLE_TREE_VIEW_ENABLED.

diff --git a/Engine/Plugins/Editor/AssetManagerEditor/Source/AssetManagerEditor/Private/AssetManagerEditorModule.cpp b/Engine/Plugins/Editor/AssetManagerEditor/Source/AssetManagerEditor/Private/AssetManagerEditorModule.cpp
--- a/Engine/Plugins/Editor/AssetManagerEditor/Source/AssetManagerEditor/Private/AssetManagerEditorModule.cpp
+++ b/Engine/Plugins/Editor/AssetManagerEditor/Source/AssetManagerEditor/Private/AssetManagerEditorModule.cpp
@@ -59,6 +59,10 @@
 #include "ContentBrowserMenuContexts.h"
 #include "IContentBrowserDataModule.h"
 #include "ContentBrowserDataSubsystem.h"
+#if ASSET_TABLE_TREE_VIEW_ENABLED
+#include "Insights/Common/InsightsStyle.h"
+#include "Insights/Filter/ViewModels/Filters.h"
+#endif //ASSET_TABLE_TREE_VIEW_ENABLED
 
 #define LOCTEXT_NAMESPACE "AssetManagerEditor"
 
@@ -425,6 +429,11 @@ void FAssetManagerEditorModule::StartupModule()
 
 	if (GIsEditor && !IsRunningCommandlet())
 	{
+#if ASSET_TABLE_TREE_VIEW_ENABLED
+		UE::Insights::FInsightsStyle::Initialize();
+		UE::Insights::FFilterService::Initialize();
+#endif //ASSET_TABLE_TREE_VIEW_ENABLED
+
 		AuditCmds.Add(IConsoleManager::Get().RegisterConsoleCommand(
 			TEXT("AssetManager.AssetAudit"),
 			TEXT("Dumps statistics about assets to the log."),
@@ -599,17 +608,35 @@ void FAssetManagerEditorModule::ShutdownModule()
 		// Cleanup tool menus
 		UToolMenus::UnRegisterStartupCallback(this);
 		UToolMenus::UnregisterOwner(this);
+
+#if ASSET_TABLE_TREE_VIEW_ENABLED
+		UE::Insights::FFilterService::Shutdown();
+		UE::Insights::FInsightsStyle::Shutdown();
+#endif //ASSET_TABLE_TREE_VIEW_ENABLED
 	}
 }
 
 TSharedRef<SDockTab> FAssetManagerEditorModule::SpawnAssetAuditTab(const FSpawnTabArgs& Args)
 {
 	check(UAssetManager::IsInitialized());
-	return SAssignNew(AssetAuditTab, SDockTab)
+
+	TSharedRef<SDockTab> DockTab = SAssignNew(AssetAuditTab, SDockTab)
 		.TabRole(ETabRole::NomadTab)
 		[
 			SAssignNew(AssetAuditUI, SAssetAuditBrowser)
 		];
+
+#if ASSET_TABLE_TREE_VIEW_ENABLED
+	DockTab->SetOnTabClosed(SDockTab::FOnTabClosedCallback::CreateLambda([this](TSharedRef<SDockTab>)
+		{
+			if (AssetAuditUI.IsValid())
+			{
+				AssetAuditUI.Pin()->OnClose();
+			}
+		}));
+#endif // ASSET_TABLE_TREE_VIEW_ENABLED
+
+	return DockTab;
 }
 
 TSharedRef<SDockTab> FAssetManagerEditorModule::SpawnReferenceViewerTab(const FSpawnTabArgs& Args)
diff --git a/Engine/Plugins/Editor/AssetManagerEditor/Source/AssetManagerEditor/Private/SAssetAuditBrowser.cpp b/Engine/Plugins/Editor/AssetManagerEditor/Source/AssetManagerEditor/Private/SAssetAuditBrowser.cpp
--- a/Engine/Plugins/Editor/AssetManagerEditor/Source/AssetManagerEditor/Private/SAssetAuditBrowser.cpp
+++ b/Engine/Plugins/Editor/AssetManagerEditor/Source/AssetManagerEditor/Private/SAssetAuditBrowser.cpp
@@ -24,6 +24,11 @@
 #include "Subsystems/AssetEditorSubsystem.h"
 #include "Editor.h"
 
+#if ASSET_TABLE_TREE_VIEW_ENABLED
+#include "TreeView/AssetTable.h"
+#include "TreeView/SAssetTableTreeView.h"
+#endif
+
 #define LOCTEXT_NAMESPACE "AssetManagementBrowser"
 
 const FString SAssetAuditBrowser::SettingsIniSection = TEXT("AssetManagementBrowser");
@@ -350,6 +355,12 @@ void SAssetAuditBrowser::Construct(const FArguments& InArgs)
 			]
 		];
 
+#if ASSET_TABLE_TREE_VIEW_ENABLED
+	TSharedRef<FAssetTable> AssetTable = MakeShared<FAssetTable>();
+	AssetTable->Reset();
+	AssetTable->SetDisplayName(FText::FromString(TEXT("AssetTable")));
+#endif //ASSET_TABLE_TREE_VIEW_ENABLED
+
 	this->ChildSlot
 	[
 		SNew(SVerticalBox)
@@ -503,8 +514,24 @@ void SAssetAuditBrowser::Construct(const FArguments& InArgs)
 				ContentBrowserModule.Get().CreateAssetPicker(Config)
 			]
 		]
+#if ASSET_TABLE_TREE_VIEW_ENABLED
+		+SVerticalBox::Slot()
+		.FillHeight(1.f)
+		[
+			SNew(SBorder)
+			.Padding(FMargin(3))
+			.BorderImage(FAppStyle::GetBrush("ToolPanel.GroupBorder"))
+			[
+				SAssignNew(AssetTableTreeView, SAssetTableTreeView, AssetTable)
+			]
+		]
+#endif //ASSET_TABLE_TREE_VIEW_ENABLED
 	];
 
+#if ASSET_TABLE_TREE_VIEW_ENABLED
+	AssetTableTreeView->RebuildTreeAsync(); // only needed for initial mock data
+#endif
+
 	// Refresh our data sources to make sure we are displaying up to date info otherwise the user will need to 
 	// hit "refresh" every time they open the dialog to be able to trust it
 	ManagerEditorModule.RefreshRegistryData();
@@ -715,6 +742,40 @@ void SAssetAuditBrowser::RefreshAssetView()
 	}
 
 	SetFilterDelegate.Execute(Filter);
+
+#if ASSET_TABLE_TREE_VIEW_ENABLED
+	// Refresh list of assets for the tree view
+	if (AssetTableTreeView.IsValid())
+	{
+		TSharedPtr<FAssetTable> AssetTable = AssetTableTreeView->GetAssetTable();
+		if (AssetTable.IsValid())
+		{
+			TArray<FAssetTableRow>& Assets = AssetTable->GetAssets();
+			Assets.Reset();
+
+			TArray<FAssetData> SourceAssets;
+			AssetRegistry->GetAssets(Filter, SourceAssets);
+			for (const FAssetData& SourceAsset : SourceAssets)
+			{
+				FAssetTableRow& Asset = Assets.AddDefaulted_GetRef();
+				Asset.Type = SourceAsset.AssetClassPath.ToString();
+				Asset.Name = SourceAsset.AssetName.ToString();
+				Asset.Path = SourceAsset.GetObjectPathString();
+				EditorModule->GetStringValueForCustomColumn(SourceAsset, FPrimaryAssetId::PrimaryAssetTypeTag, Asset.PrimaryType);
+				EditorModule->GetStringValueForCustomColumn(SourceAsset, FPrimaryAssetId::PrimaryAssetNameTag, Asset.PrimaryName);
+				EditorModule->GetIntegerValueForCustomColumn(SourceAsset, IAssetManagerEditorModule::ManagedDiskSizeName, Asset.ManagedDiskSize);
+				EditorModule->GetIntegerValueForCustomColumn(SourceAsset, IAssetManagerEditorModule::DiskSizeName, Asset.DiskSize);
+				EditorModule->GetIntegerValueForCustomColumn(SourceAsset, IAssetManagerEditorModule::StageChunkCompressedSizeName, Asset.StagedCompressedSize);
+				EditorModule->GetIntegerValueForCustomColumn(SourceAsset, IAssetManagerEditorModule::TotalUsageName, Asset.TotalUsageCount);
+				EditorModule->GetStringValueForCustomColumn(SourceAsset, IAssetManagerEditorModule::CookRuleName, Asset.CookRule);
+				EditorModule->GetStringValueForCustomColumn(SourceAsset, IAssetManagerEditorModule::ChunksName, Asset.Chunks);
+				Asset.NativeClass = SourceAsset.AssetClassPath.GetAssetName().ToString();
+			}
+
+			AssetTableTreeView->RebuildTreeAsync();
+		}
+	}
+#endif // ASSET_TABLE_TREE_VIEW_ENABLED
 }
 
 FReply SAssetAuditBrowser::ClearAssets()
@@ -910,4 +971,14 @@ void SAssetAuditBrowser::SetCurrentRegistrySource(const FAssetManagerEditorRegis
 	RefreshAssetView();
 }
 
+void SAssetAuditBrowser::OnClose()
+{
+#if ASSET_TABLE_TREE_VIEW_ENABLED
+	if (AssetTableTreeView.IsValid())
+	{
+		AssetTableTreeView->OnClose();
+	}
+#endif // ASSET_TABLE_TREE_VIEW_ENABLED
+}
+
 #undef LOCTEXT_NAMESPACE
diff --git a/Engine/Plugins/Editor/AssetManagerEditor/Source/AssetManagerEditor/Private/SAssetAuditBrowser.h b/Engine/Plugins/Editor/AssetManagerEditor/Source/AssetManagerEditor/Private/SAssetAuditBrowser.h
--- a/Engine/Plugins/Editor/AssetManagerEditor/Source/AssetManagerEditor/Private/SAssetAuditBrowser.h
+++ b/Engine/Plugins/Editor/AssetManagerEditor/Source/AssetManagerEditor/Private/SAssetAuditBrowser.h
@@ -5,6 +5,8 @@
 #include "ContentBrowserDelegates.h"
 #include "Widgets/SCompoundWidget.h"
 
+#define ASSET_TABLE_TREE_VIEW_ENABLED 0 // switch for work in progress / experimental code; to be removed once ready
+
 struct FAssetIdentifier;
 struct FAssetManagerEditorRegistrySource;
 
@@ -36,6 +38,9 @@ public:
 	/** Called when the current registry source changes */
 	void SetCurrentRegistrySource(const FAssetManagerEditorRegistrySource* RegistrySource);
 
+	/** Called when tab is closed */
+	void OnClose();
+
 protected:
 	virtual FReply OnKeyDown(const FGeometry& MyGeometry, const FKeyEvent& InKeyEvent) override;
 	virtual FReply OnDragOver(const FGeometry& MyGeometry, const FDragDropEvent& DragDropEvent) override;
@@ -166,4 +171,8 @@ protected:
 	class IAssetRegistry* AssetRegistry;
 	class UAssetManager* AssetManager;
 	class IAssetManagerEditorModule* EditorModule;
+
+#if ASSET_TABLE_TREE_VIEW_ENABLED
+	TSharedPtr<class SAssetTableTreeView> AssetTableTreeView;
+#endif
 };
diff --git a/Engine/Plugins/Editor/AssetManagerEditor/Source/AssetManagerEditor/Private/TreeView/AssetTable.cpp b/Engine/Plugins/Editor/AssetManagerEditor/Source/AssetManagerEditor/Private/TreeView/AssetTable.cpp
new file mode 100644
--- /dev/null
+++ b/Engine/Plugins/Editor/AssetManagerEditor/Source/AssetManagerEditor/Private/TreeView/AssetTable.cpp
@@ -0,0 +1,841 @@
+// Copyright Epic Games, Inc. All Rights Reserved.
+
+#include "AssetTable.h"
+
+#include "Styling/StyleColors.h"
+#include "Widgets/SBoxPanel.h"
+#include "Widgets/SToolTip.h"
+#include "Widgets/Text/STextBlock.h"
+
+// AssetManagerEditor
+#include "AssetTreeNode.h"
+#include "Insights/Table/ViewModels/TableCellValueFormatter.h"
+#include "Insights/Table/ViewModels/TableCellValueGetter.h"
+#include "Insights/Table/ViewModels/TableCellValueSorter.h"
+#include "Insights/Table/ViewModels/TableColumn.h"
+
+#define LOCTEXT_NAMESPACE "FAssetTable"
+
+////////////////////////////////////////////////////////////////////////////////////////////////////
+// Column identifiers
+
+const FName FAssetTableColumns::CountColumnId(TEXT("Count"));
+const FName FAssetTableColumns::NameColumnId(TEXT("Name"));
+const FName FAssetTableColumns::TypeColumnId(TEXT("Type"));
+const FName FAssetTableColumns::PathColumnId(TEXT("Path"));
+const FName FAssetTableColumns::PrimaryTypeColumnId(TEXT("PrimaryType"));
+const FName FAssetTableColumns::PrimaryNameColumnId(TEXT("PrimaryName"));
+const FName FAssetTableColumns::ManagedDiskSizeColumnId(TEXT("ManagedDiskSize"));
+const FName FAssetTableColumns::DiskSizeColumnId(TEXT("DiskSize"));
+const FName FAssetTableColumns::StagedCompressedSizeColumnId(TEXT("StagedCompressedSize"));
+const FName FAssetTableColumns::TotalUsageCountColumnId(TEXT("TotalUsageCount"));
+const FName FAssetTableColumns::CookRuleColumnId(TEXT("CookRule"));
+const FName FAssetTableColumns::ChunksColumnId(TEXT("Chunks"));
+const FName FAssetTableColumns::NativeClassColumnId(TEXT("NativeClass"));
+
+////////////////////////////////////////////////////////////////////////////////////////////////////
+// FAssetTable
+////////////////////////////////////////////////////////////////////////////////////////////////////
+
+FAssetTable::FAssetTable()
+{
+#if 0 // debug, mock data
+	constexpr int AssetCount = 100;
+	Assets.Reserve(AssetCount);
+	for (int AssetIndex = 0; AssetIndex < AssetCount; ++AssetIndex)
+	{
+		FAssetTableRow& Asset = Assets.AddDefaulted_GetRef();
+		
+		int32 Id = FMath::Rand();
+		int32 Id2 = FMath::Rand();
+
+		Asset.Type = FString::Printf(TEXT("Type%02d"), Id % 10);
+		Asset.Name = FString::Printf(TEXT("Name%d"), Id);
+		Asset.Path = FString::Printf(TEXT("A%02d/B%02d/C%02d/D%02d"), Id % 11, Id % 13, Id % 17, Id % 19);
+		Asset.PrimaryType = FString::Printf(TEXT("PT_%02d"), Id2 % 10);
+		Asset.PrimaryName = FString::Printf(TEXT("PN%d"), Id2);
+		Asset.ManagedDiskSize = FMath::Abs(Id * Id);
+		Asset.DiskSize = FMath::Abs(Id * Id * Id);
+		Asset.StagedCompressedSize = Asset.DiskSize / 2;
+		Asset.TotalUsageCount = Id % 1000;
+		Asset.CookRule = FString::Printf(TEXT("CookRule%02d"), (Id * Id) % 8);
+		Asset.Chunks = FString::Printf(TEXT("Chunks%02d"), (Id * Id + 1) % 41);
+		Asset.NativeClass = FString::Printf(TEXT("NativeClass%02d"), (Id * Id * Id) % 8);
+	}
+#endif
+}
+
+////////////////////////////////////////////////////////////////////////////////////////////////////
+
+FAssetTable::~FAssetTable()
+{
+}
+
+////////////////////////////////////////////////////////////////////////////////////////////////////
+
+void FAssetTable::Reset()
+{
+	//...
+
+	FTable::Reset();
+
+	AddDefaultColumns();
+}
+
+////////////////////////////////////////////////////////////////////////////////////////////////////
+
+void FAssetTable::AddDefaultColumns()
+{
+	using namespace UE::Insights;
+
+	//////////////////////////////////////////////////
+	// Hierarchy Column
+	{
+		const int32 HierarchyColumnIndex = -1;
+		const TCHAR* HierarchyColumnName = nullptr;
+		AddHierarchyColumn(HierarchyColumnIndex, HierarchyColumnName);
+
+		const TSharedRef<FTableColumn>& ColumnRef = GetColumns()[0];
+		ColumnRef->SetInitialWidth(200.0f);
+		ColumnRef->SetShortName(LOCTEXT("HierarchyColumnName", "Hierarchy"));
+		ColumnRef->SetTitleName(LOCTEXT("HierarchyColumnTitle", "Asset Hierarchy"));
+		ColumnRef->SetDescription(LOCTEXT("HierarchyColumnDesc", "Hierarchy of the asset tree"));
+	}
+
+	int32 ColumnIndex = 0;
+
+	//////////////////////////////////////////////////
+	// Count Column
+	{
+		TSharedRef<FTableColumn> ColumnRef = MakeShared<FTableColumn>(FAssetTableColumns::CountColumnId);
+		FTableColumn& Column = *ColumnRef;
+
+		Column.SetIndex(ColumnIndex++);
+
+		Column.SetShortName(LOCTEXT("CountColumnName", "Count"));
+		Column.SetTitleName(LOCTEXT("CountColumnTitle", "Asset Count"));
+		Column.SetDescription(LOCTEXT("CountColumnDesc", "Number of assets"));
+
+		Column.SetFlags(ETableColumnFlags::ShouldBeVisible | ETableColumnFlags::CanBeHidden | ETableColumnFlags::CanBeFiltered);
+
+		Column.SetHorizontalAlignment(HAlign_Right);
+		Column.SetInitialWidth(100.0f);
+
+		Column.SetDataType(ETableCellDataType::Int64);
+
+		class FAssetCountValueGetter : public FTableCellValueGetter
+		{
+		public:
+			virtual const TOptional<FTableCellValue> GetValue(const FTableColumn& Column, const FBaseTreeNode& Node) const override
+			{
+				if (Node.IsGroup())
+				{
+					const FTableTreeNode& NodePtr = static_cast<const FTableTreeNode&>(Node);
+					if (NodePtr.HasAggregatedValue(Column.GetId()))
+					{
+						return NodePtr.GetAggregatedValue(Column.GetId());
+					}
+				}
+				else //if (Node->Is<FAssetTreeNode>())
+				{
+					//const FAssetTreeNode& TreeNode = static_cast<const FAssetTreeNode&>(Node);
+					//const FAssetTableRow& Asset = TreeNode.GetAssetChecked();
+					return FTableCellValue((int64)1);
+				}
+
+				return TOptional<FTableCellValue>();
+			}
+		};
+		TSharedRef<ITableCellValueGetter> Getter = MakeShared<FAssetCountValueGetter>();
+		Column.SetValueGetter(Getter);
+
+		TSharedRef<ITableCellValueFormatter> Formatter = MakeShared<FInt64ValueFormatterAsNumber>();
+		Column.SetValueFormatter(Formatter);
+
+		TSharedRef<ITableCellValueSorter> Sorter = MakeShared<FSorterByInt64Value>(ColumnRef);
+		Column.SetValueSorter(Sorter);
+		Column.SetInitialSortMode(EColumnSortMode::Descending);
+
+		Column.SetAggregation(ETableColumnAggregation::Sum);
+
+		AddColumn(ColumnRef);
+	}
+	//////////////////////////////////////////////////
+	// Name Column
+	{
+		TSharedRef<FTableColumn> ColumnRef = MakeShared<FTableColumn>(FAssetTableColumns::NameColumnId);
+		FTableColumn& Column = *ColumnRef;
+
+		Column.SetIndex(ColumnIndex++);
+
+		Column.SetShortName(LOCTEXT("NameColumnName", "Name"));
+		Column.SetTitleName(LOCTEXT("NameColumnTitle", "Name"));
+		Column.SetDescription(LOCTEXT("NameColumnDesc", "Asset's name"));
+
+		Column.SetFlags(ETableColumnFlags::ShouldBeVisible | ETableColumnFlags::CanBeHidden | ETableColumnFlags::CanBeFiltered);
+
+		Column.SetHorizontalAlignment(HAlign_Left);
+		Column.SetInitialWidth(120.0f);
+
+		Column.SetDataType(ETableCellDataType::CString);
+
+		class FAssetNameValueGetter : public FTableCellValueGetter
+		{
+		public:
+			virtual const TOptional<FTableCellValue> GetValue(const FTableColumn& Column, const FBaseTreeNode& Node) const
+			{
+				if (Node.IsGroup())
+				{
+					const FTableTreeNode& NodePtr = static_cast<const FTableTreeNode&>(Node);
+					if (NodePtr.HasAggregatedValue(Column.GetId()))
+					{
+						return NodePtr.GetAggregatedValue(Column.GetId());
+					}
+				}
+				else //if (Node->Is<FAssetTreeNode>())
+				{
+					const FAssetTreeNode& TreeNode = static_cast<const FAssetTreeNode&>(Node);
+					const FAssetTableRow& Asset = TreeNode.GetAssetChecked();
+					return FTableCellValue(Asset.GetName());
+				}
+
+				return TOptional<FTableCellValue>();
+			}
+		};
+		TSharedRef<ITableCellValueGetter> Getter = MakeShared<FAssetNameValueGetter>();
+		Column.SetValueGetter(Getter);
+
+		TSharedRef<ITableCellValueFormatter> Formatter = MakeShared<FCStringValueFormatterAsText>();
+		Column.SetValueFormatter(Formatter);
+
+		TSharedRef<ITableCellValueSorter> Sorter = MakeShared<FSorterByCStringValue>(ColumnRef);
+		Column.SetValueSorter(Sorter);
+
+		Column.SetAggregation(ETableColumnAggregation::SameValue);
+
+		AddColumn(ColumnRef);
+	}
+	//////////////////////////////////////////////////
+	// Type Column
+	{
+		TSharedRef<FTableColumn> ColumnRef = MakeShared<FTableColumn>(FAssetTableColumns::TypeColumnId);
+		FTableColumn& Column = *ColumnRef;
+
+		Column.SetIndex(ColumnIndex++);
+
+		Column.SetShortName(LOCTEXT("TypeColumnName", "Type"));
+		Column.SetTitleName(LOCTEXT("TypeColumnTitle", "Type"));
+		Column.SetDescription(LOCTEXT("TypeColumnDesc", "Asset's type"));
+
+		Column.SetFlags(ETableColumnFlags::ShouldBeVisible | ETableColumnFlags::CanBeHidden | ETableColumnFlags::CanBeFiltered);
+
+		Column.SetHorizontalAlignment(HAlign_Left);
+		Column.SetInitialWidth(120.0f);
+
+		Column.SetDataType(ETableCellDataType::CString);
+
+		class FAssetTypeValueGetter : public FTableCellValueGetter
+		{
+		public:
+			virtual const TOptional<FTableCellValue> GetValue(const FTableColumn& Column, const FBaseTreeNode& Node) const
+			{
+				if (Node.IsGroup())
+				{
+					const FTableTreeNode& NodePtr = static_cast<const FTableTreeNode&>(Node);
+					if (NodePtr.HasAggregatedValue(Column.GetId()))
+					{
+						return NodePtr.GetAggregatedValue(Column.GetId());
+					}
+				}
+				else //if (Node->Is<FAssetTreeNode>())
+				{
+					const FAssetTreeNode& TreeNode = static_cast<const FAssetTreeNode&>(Node);
+					const FAssetTableRow& Asset = TreeNode.GetAssetChecked();
+					return FTableCellValue(Asset.GetType());
+				}
+
+				return TOptional<FTableCellValue>();
+			}
+		};
+		TSharedRef<ITableCellValueGetter> Getter = MakeShared<FAssetTypeValueGetter>();
+		Column.SetValueGetter(Getter);
+
+		TSharedRef<ITableCellValueFormatter> Formatter = MakeShared<FCStringValueFormatterAsText>();
+		Column.SetValueFormatter(Formatter);
+
+		TSharedRef<ITableCellValueSorter> Sorter = MakeShared<FSorterByCStringValue>(ColumnRef);
+		Column.SetValueSorter(Sorter);
+
+		Column.SetAggregation(ETableColumnAggregation::SameValue);
+
+		AddColumn(ColumnRef);
+	}
+	//////////////////////////////////////////////////
+	// Path Column
+	{
+		TSharedRef<FTableColumn> ColumnRef = MakeShared<FTableColumn>(FAssetTableColumns::PathColumnId);
+		FTableColumn& Column = *ColumnRef;
+
+		Column.SetIndex(ColumnIndex++);
+
+		Column.SetShortName(LOCTEXT("PathColumnName", "Path"));
+		Column.SetTitleName(LOCTEXT("PathColumnTitle", "Path"));
+		Column.SetDescription(LOCTEXT("PathColumnDesc", "Asset's path"));
+
+		Column.SetFlags(ETableColumnFlags::ShouldBeVisible | ETableColumnFlags::CanBeHidden | ETableColumnFlags::CanBeFiltered);
+
+		Column.SetHorizontalAlignment(HAlign_Left);
+		Column.SetInitialWidth(120.0f);
+
+		Column.SetDataType(ETableCellDataType::CString);
+
+		class FAssetPathValueGetter : public FTableCellValueGetter
+		{
+		public:
+			virtual const TOptional<FTableCellValue> GetValue(const FTableColumn& Column, const FBaseTreeNode& Node) const
+			{
+				if (Node.IsGroup())
+				{
+					const FTableTreeNode& NodePtr = static_cast<const FTableTreeNode&>(Node);
+					if (NodePtr.HasAggregatedValue(Column.GetId()))
+					{
+						return NodePtr.GetAggregatedValue(Column.GetId());
+					}
+				}
+				else //if (Node->Is<FAssetTreeNode>())
+				{
+					const FAssetTreeNode& TreeNode = static_cast<const FAssetTreeNode&>(Node);
+					const FAssetTableRow& Asset = TreeNode.GetAssetChecked();
+					return FTableCellValue(Asset.GetPath());
+				}
+
+				return TOptional<FTableCellValue>();
+			}
+		};
+		TSharedRef<ITableCellValueGetter> Getter = MakeShared<FAssetPathValueGetter>();
+		Column.SetValueGetter(Getter);
+
+		TSharedRef<ITableCellValueFormatter> Formatter = MakeShared<FCStringValueFormatterAsText>();
+		Column.SetValueFormatter(Formatter);
+
+		TSharedRef<ITableCellValueSorter> Sorter = MakeShared<FSorterByCStringValue>(ColumnRef);
+		Column.SetValueSorter(Sorter);
+
+		Column.SetAggregation(ETableColumnAggregation::SameValue);
+
+		AddColumn(ColumnRef);
+	}
+	//////////////////////////////////////////////////
+	// Primary Type Column
+	{
+		TSharedRef<FTableColumn> ColumnRef = MakeShared<FTableColumn>(FAssetTableColumns::PrimaryTypeColumnId);
+		FTableColumn& Column = *ColumnRef;
+
+		Column.SetIndex(ColumnIndex++);
+
+		Column.SetShortName(LOCTEXT("PrimaryTypeColumnName", "Primary Type"));
+		Column.SetTitleName(LOCTEXT("PrimaryTypeColumnTitle", "Primary Type"));
+		Column.SetDescription(LOCTEXT("PrimaryTypeColumnDesc", "Primary Asset Type of this asset, if set"));
+
+		Column.SetFlags(ETableColumnFlags::ShouldBeVisible | ETableColumnFlags::CanBeHidden | ETableColumnFlags::CanBeFiltered);
+
+		Column.SetHorizontalAlignment(HAlign_Left);
+		Column.SetInitialWidth(120.0f);
+
+		Column.SetDataType(ETableCellDataType::CString);
+
+		class FAssetPrimaryTypeValueGetter : public FTableCellValueGetter
+		{
+		public:
+			virtual const TOptional<FTableCellValue> GetValue(const FTableColumn& Column, const FBaseTreeNode& Node) const
+			{
+				if (Node.IsGroup())
+				{
+					const FTableTreeNode& NodePtr = static_cast<const FTableTreeNode&>(Node);
+					if (NodePtr.HasAggregatedValue(Column.GetId()))
+					{
+						return NodePtr.GetAggregatedValue(Column.GetId());
+					}
+				}
+				else //if (Node->Is<FAssetTreeNode>())
+				{
+					const FAssetTreeNode& TreeNode = static_cast<const FAssetTreeNode&>(Node);
+					const FAssetTableRow& Asset = TreeNode.GetAssetChecked();
+					return FTableCellValue(Asset.GetPrimaryType());
+				}
+
+				return TOptional<FTableCellValue>();
+			}
+		};
+		TSharedRef<ITableCellValueGetter> Getter = MakeShared<FAssetPrimaryTypeValueGetter>();
+		Column.SetValueGetter(Getter);
+
+		TSharedRef<ITableCellValueFormatter> Formatter = MakeShared<FCStringValueFormatterAsText>();
+		Column.SetValueFormatter(Formatter);
+
+		TSharedRef<ITableCellValueSorter> Sorter = MakeShared<FSorterByCStringValue>(ColumnRef);
+		Column.SetValueSorter(Sorter);
+
+		Column.SetAggregation(ETableColumnAggregation::SameValue);
+
+		AddColumn(ColumnRef);
+	}
+	//////////////////////////////////////////////////
+	// Primary Name Column
+	{
+		TSharedRef<FTableColumn> ColumnRef = MakeShared<FTableColumn>(FAssetTableColumns::PrimaryNameColumnId);
+		FTableColumn& Column = *ColumnRef;
+
+		Column.SetIndex(ColumnIndex++);
+
+		Column.SetShortName(LOCTEXT("PrimaryNameColumnName", "Primary Name"));
+		Column.SetTitleName(LOCTEXT("PrimaryNameColumnTitle", "Primary Name"));
+		Column.SetDescription(LOCTEXT("PrimaryNameColumnDesc", "Primary Asset Name of this asset, if set"));
+
+		Column.SetFlags(ETableColumnFlags::ShouldBeVisible | ETableColumnFlags::CanBeHidden | ETableColumnFlags::CanBeFiltered);
+
+		Column.SetHorizontalAlignment(HAlign_Left);
+		Column.SetInitialWidth(120.0f);
+
+		Column.SetDataType(ETableCellDataType::CString);
+
+		class FAssetPrimaryNameValueGetter : public FTableCellValueGetter
+		{
+		public:
+			virtual const TOptional<FTableCellValue> GetValue(const FTableColumn& Column, const FBaseTreeNode& Node) const
+			{
+				if (Node.IsGroup())
+				{
+					const FTableTreeNode& NodePtr = static_cast<const FTableTreeNode&>(Node);
+					if (NodePtr.HasAggregatedValue(Column.GetId()))
+					{
+						return NodePtr.GetAggregatedValue(Column.GetId());
+					}
+				}
+				else //if (Node->Is<FAssetTreeNode>())
+				{
+					const FAssetTreeNode& TreeNode = static_cast<const FAssetTreeNode&>(Node);
+					const FAssetTableRow& Asset = TreeNode.GetAssetChecked();
+					return FTableCellValue(Asset.GetPrimaryName());
+				}
+
+				return TOptional<FTableCellValue>();
+			}
+		};
+		TSharedRef<ITableCellValueGetter> Getter = MakeShared<FAssetPrimaryNameValueGetter>();
+		Column.SetValueGetter(Getter);
+
+		TSharedRef<ITableCellValueFormatter> Formatter = MakeShared<FCStringValueFormatterAsText>();
+		Column.SetValueFormatter(Formatter);
+
+		TSharedRef<ITableCellValueSorter> Sorter = MakeShared<FSorterByCStringValue>(ColumnRef);
+		Column.SetValueSorter(Sorter);
+
+		Column.SetAggregation(ETableColumnAggregation::SameValue);
+
+		AddColumn(ColumnRef);
+	}
+	//////////////////////////////////////////////////
+	// Managed Disk Size Column
+	{
+		TSharedRef<FTableColumn> ColumnRef = MakeShared<FTableColumn>(FAssetTableColumns::ManagedDiskSizeColumnId);
+		FTableColumn& Column = *ColumnRef;
+
+		Column.SetIndex(ColumnIndex++);
+
+		Column.SetShortName(LOCTEXT("ManagedDiskSizeColumnName", "Disk Size"));
+		Column.SetTitleName(LOCTEXT("ManagedDiskSizeColumnTitle", "Disk Size"));
+		Column.SetDescription(LOCTEXT("ManagedDiskSizeColumnDesc", "Total disk space used by both this and all managed assets"));
+
+		Column.SetFlags(ETableColumnFlags::ShouldBeVisible | ETableColumnFlags::CanBeHidden | ETableColumnFlags::CanBeFiltered);
+
+		Column.SetHorizontalAlignment(HAlign_Right);
+		Column.SetInitialWidth(100.0f);
+
+		Column.SetDataType(ETableCellDataType::Int64);
+
+		class FManagedDiskSizeValueGetter : public FTableCellValueGetter
+		{
+		public:
+			virtual const TOptional<FTableCellValue> GetValue(const FTableColumn& Column, const FBaseTreeNode& Node) const override
+			{
+				if (Node.IsGroup())
+				{
+					const FTableTreeNode& NodePtr = static_cast<const FTableTreeNode&>(Node);
+					if (NodePtr.HasAggregatedValue(Column.GetId()))
+					{
+						return NodePtr.GetAggregatedValue(Column.GetId());
+					}
+				}
+				else //if (Node->Is<FAssetTreeNode>())
+				{
+					const FAssetTreeNode& TreeNode = static_cast<const FAssetTreeNode&>(Node);
+					const FAssetTableRow& Asset = TreeNode.GetAssetChecked();
+					return FTableCellValue(static_cast<int64>(Asset.GetManagedDiskSize()));
+				}
+
+				return TOptional<FTableCellValue>();
+			}
+		};
+		TSharedRef<ITableCellValueGetter> Getter = MakeShared<FManagedDiskSizeValueGetter>();
+		Column.SetValueGetter(Getter);
+
+		TSharedRef<ITableCellValueFormatter> Formatter = MakeShared<FInt64ValueFormatterAsMemory>();
+		Column.SetValueFormatter(Formatter);
+
+		TSharedRef<ITableCellValueSorter> Sorter = MakeShared<FSorterByInt64Value>(ColumnRef);
+		Column.SetValueSorter(Sorter);
+		Column.SetInitialSortMode(EColumnSortMode::Descending);
+
+		//TSharedRef<IFilterValueConverter> Converter = MakeShared<FMemoryFilterValueConverter>();
+		//Column.SetValueConverter(Converter);
+
+		Column.SetAggregation(ETableColumnAggregation::Sum);
+
+		AddColumn(ColumnRef);
+	}
+	//////////////////////////////////////////////////
+	// Exclusive Disk Size Column
+	{
+		TSharedRef<FTableColumn> ColumnRef = MakeShared<FTableColumn>(FAssetTableColumns::DiskSizeColumnId);
+		FTableColumn& Column = *ColumnRef;
+
+		Column.SetIndex(ColumnIndex++);
+
+		Column.SetShortName(LOCTEXT("DiskSizeColumnName", "Exclusive Disk Size"));
+		Column.SetTitleName(LOCTEXT("DiskSizeColumnTitle", "Exclusive Disk Size"));
+		Column.SetDescription(LOCTEXT("DiskSizeColumnDesc", "Size of saved file(s) on disk for this asset's package. If Asset Registry Writeback is enabled and a platform is selected, this will be the uncompressed size of the iostore chunks for the asset's package."));
+
+		Column.SetFlags(ETableColumnFlags::ShouldBeVisible | ETableColumnFlags::CanBeHidden | ETableColumnFlags::CanBeFiltered);
+
+		Column.SetHorizontalAlignment(HAlign_Right);
+		Column.SetInitialWidth(100.0f);
+
+		Column.SetDataType(ETableCellDataType::Int64);
+
+		class FDiskSizeValueGetter : public FTableCellValueGetter
+		{
+		public:
+			virtual const TOptional<FTableCellValue> GetValue(const FTableColumn& Column, const FBaseTreeNode& Node) const override
+			{
+				if (Node.IsGroup())
+				{
+					const FTableTreeNode& NodePtr = static_cast<const FTableTreeNode&>(Node);
+					if (NodePtr.HasAggregatedValue(Column.GetId()))
+					{
+						return NodePtr.GetAggregatedValue(Column.GetId());
+					}
+				}
+				else //if (Node->Is<FAssetTreeNode>())
+				{
+					const FAssetTreeNode& TreeNode = static_cast<const FAssetTreeNode&>(Node);
+					const FAssetTableRow& Asset = TreeNode.GetAssetChecked();
+					return FTableCellValue(static_cast<int64>(Asset.GetDiskSize()));
+				}
+
+				return TOptional<FTableCellValue>();
+			}
+		};
+		TSharedRef<ITableCellValueGetter> Getter = MakeShared<FDiskSizeValueGetter>();
+		Column.SetValueGetter(Getter);
+
+		TSharedRef<ITableCellValueFormatter> Formatter = MakeShared<FInt64ValueFormatterAsMemory>();
+		Column.SetValueFormatter(Formatter);
+
+		TSharedRef<ITableCellValueSorter> Sorter = MakeShared<FSorterByInt64Value>(ColumnRef);
+		Column.SetValueSorter(Sorter);
+		Column.SetInitialSortMode(EColumnSortMode::Descending);
+
+		//TSharedRef<IFilterValueConverter> Converter = MakeShared<FMemoryFilterValueConverter>();
+		//Column.SetValueConverter(Converter);
+
+		Column.SetAggregation(ETableColumnAggregation::Sum);
+
+		AddColumn(ColumnRef);
+	}
+	//////////////////////////////////////////////////
+	// Staged Compressed Size Column
+	{
+		TSharedRef<FTableColumn> ColumnRef = MakeShared<FTableColumn>(FAssetTableColumns::StagedCompressedSizeColumnId);
+		FTableColumn& Column = *ColumnRef;
+
+		Column.SetIndex(ColumnIndex++);
+
+		Column.SetShortName(LOCTEXT("StagedCompressedSizeColumnName", "Staged Compressed Size"));
+		Column.SetTitleName(LOCTEXT("StagedCompressedSizeColumnTitle", "Staged Compressed Size"));
+		Column.SetDescription(LOCTEXT("StagedCompressedSizeColumnDesc", "Compressed size of iostore chunks for this asset's package. Only visible after staging."));
+
+		Column.SetFlags(ETableColumnFlags::ShouldBeVisible | ETableColumnFlags::CanBeHidden | ETableColumnFlags::CanBeFiltered);
+
+		Column.SetHorizontalAlignment(HAlign_Right);
+		Column.SetInitialWidth(100.0f);
+
+		Column.SetDataType(ETableCellDataType::Int64);
+
+		class FStagedCompressedSizeValueGetter : public FTableCellValueGetter
+		{
+		public:
+			virtual const TOptional<FTableCellValue> GetValue(const FTableColumn& Column, const FBaseTreeNode& Node) const override
+			{
+				if (Node.IsGroup())
+				{
+					const FTableTreeNode& NodePtr = static_cast<const FTableTreeNode&>(Node);
+					if (NodePtr.HasAggregatedValue(Column.GetId()))
+					{
+						return NodePtr.GetAggregatedValue(Column.GetId());
+					}
+				}
+				else //if (Node->Is<FAssetTreeNode>())
+				{
+					const FAssetTreeNode& TreeNode = static_cast<const FAssetTreeNode&>(Node);
+					const FAssetTableRow& Asset = TreeNode.GetAssetChecked();
+					return FTableCellValue(static_cast<int64>(Asset.GetStagedCompressedSize()));
+				}
+
+				return TOptional<FTableCellValue>();
+			}
+		};
+		TSharedRef<ITableCellValueGetter> Getter = MakeShared<FStagedCompressedSizeValueGetter>();
+		Column.SetValueGetter(Getter);
+
+		TSharedRef<ITableCellValueFormatter> Formatter = MakeShared<FInt64ValueFormatterAsMemory>();
+		Column.SetValueFormatter(Formatter);
+
+		TSharedRef<ITableCellValueSorter> Sorter = MakeShared<FSorterByInt64Value>(ColumnRef);
+		Column.SetValueSorter(Sorter);
+		Column.SetInitialSortMode(EColumnSortMode::Descending);
+
+		//TSharedRef<IFilterValueConverter> Converter = MakeShared<FMemoryFilterValueConverter>();
+		//Column.SetValueConverter(Converter);
+
+		Column.SetAggregation(ETableColumnAggregation::Sum);
+
+		AddColumn(ColumnRef);
+	}
+	//////////////////////////////////////////////////
+	// Total Usage Count Column
+	{
+		TSharedRef<FTableColumn> ColumnRef = MakeShared<FTableColumn>(FAssetTableColumns::TotalUsageCountColumnId);
+		FTableColumn& Column = *ColumnRef;
+
+		Column.SetIndex(ColumnIndex++);
+
+		Column.SetShortName(LOCTEXT("TotalUsageCountColumnName", "Total Usage"));
+		Column.SetTitleName(LOCTEXT("TotalUsageCountColumnTitle", "Total Usage Count"));
+		Column.SetDescription(LOCTEXT("TotalUsageCountColumnDesc", "Weighted count of Primary Assets that use this\nA higher usage means it's more likely to be in memory at runtime."));
+
+		Column.SetFlags(ETableColumnFlags::ShouldBeVisible | ETableColumnFlags::CanBeHidden | ETableColumnFlags::CanBeFiltered);
+
+		Column.SetHorizontalAlignment(HAlign_Right);
+		Column.SetInitialWidth(100.0f);
+
+		Column.SetDataType(ETableCellDataType::Int64);
+
+		class FTotalUsageCountValueGetter : public FTableCellValueGetter
+		{
+		public:
+			virtual const TOptional<FTableCellValue> GetValue(const FTableColumn& Column, const FBaseTreeNode& Node) const override
+			{
+				if (Node.IsGroup())
+				{
+					const FTableTreeNode& NodePtr = static_cast<const FTableTreeNode&>(Node);
+					if (NodePtr.HasAggregatedValue(Column.GetId()))
+					{
+						return NodePtr.GetAggregatedValue(Column.GetId());
+					}
+				}
+				else //if (Node->Is<FAssetTreeNode>())
+				{
+					const FAssetTreeNode& TreeNode = static_cast<const FAssetTreeNode&>(Node);
+					const FAssetTableRow& Asset = TreeNode.GetAssetChecked();
+					return FTableCellValue(static_cast<int64>(Asset.GetTotalUsageCount()));
+				}
+
+				return TOptional<FTableCellValue>();
+			}
+		};
+		TSharedRef<ITableCellValueGetter> Getter = MakeShared<FTotalUsageCountValueGetter>();
+		Column.SetValueGetter(Getter);
+
+		TSharedRef<ITableCellValueFormatter> Formatter = MakeShared<FInt64ValueFormatterAsNumber>();
+		Column.SetValueFormatter(Formatter);
+
+		TSharedRef<ITableCellValueSorter> Sorter = MakeShared<FSorterByInt64Value>(ColumnRef);
+		Column.SetValueSorter(Sorter);
+		Column.SetInitialSortMode(EColumnSortMode::Descending);
+
+		Column.SetAggregation(ETableColumnAggregation::Sum);
+
+		AddColumn(ColumnRef);
+	}
+	//////////////////////////////////////////////////
+	// Cook Rule Column
+	{
+		TSharedRef<FTableColumn> ColumnRef = MakeShared<FTableColumn>(FAssetTableColumns::CookRuleColumnId);
+		FTableColumn& Column = *ColumnRef;
+
+		Column.SetIndex(ColumnIndex++);
+
+		Column.SetShortName(LOCTEXT("CookRuleColumnName", "Cook Rule"));
+		Column.SetTitleName(LOCTEXT("CookRuleColumnTitle", "Cook Rule"));
+		Column.SetDescription(LOCTEXT("CookRuleColumnDesc", "Whether this asset will be cooked or not"));
+
+		Column.SetFlags(ETableColumnFlags::ShouldBeVisible | ETableColumnFlags::CanBeHidden | ETableColumnFlags::CanBeFiltered);
+
+		Column.SetHorizontalAlignment(HAlign_Left);
+		Column.SetInitialWidth(120.0f);
+
+		Column.SetDataType(ETableCellDataType::CString);
+
+		class FCookRuleValueGetter : public FTableCellValueGetter
+		{
+		public:
+			virtual const TOptional<FTableCellValue> GetValue(const FTableColumn& Column, const FBaseTreeNode& Node) const
+			{
+				if (Node.IsGroup())
+				{
+					const FTableTreeNode& NodePtr = static_cast<const FTableTreeNode&>(Node);
+					if (NodePtr.HasAggregatedValue(Column.GetId()))
+					{
+						return NodePtr.GetAggregatedValue(Column.GetId());
+					}
+				}
+				else //if (Node->Is<FAssetTreeNode>())
+				{
+					const FAssetTreeNode& TreeNode = static_cast<const FAssetTreeNode&>(Node);
+					const FAssetTableRow& Asset = TreeNode.GetAssetChecked();
+					return FTableCellValue(Asset.GetChunks());
+				}
+
+				return TOptional<FTableCellValue>();
+			}
+		};
+		TSharedRef<ITableCellValueGetter> Getter = MakeShared<FCookRuleValueGetter>();
+		Column.SetValueGetter(Getter);
+
+		TSharedRef<ITableCellValueFormatter> Formatter = MakeShared<FCStringValueFormatterAsText>();
+		Column.SetValueFormatter(Formatter);
+
+		TSharedRef<ITableCellValueSorter> Sorter = MakeShared<FSorterByCStringValue>(ColumnRef);
+		Column.SetValueSorter(Sorter);
+
+		Column.SetAggregation(ETableColumnAggregation::SameValue);
+
+		AddColumn(ColumnRef);
+	}
+	//////////////////////////////////////////////////
+	// Chunks Column
+	{
+		TSharedRef<FTableColumn> ColumnRef = MakeShared<FTableColumn>(FAssetTableColumns::ChunksColumnId);
+		FTableColumn& Column = *ColumnRef;
+
+		Column.SetIndex(ColumnIndex++);
+
+		Column.SetShortName(LOCTEXT("ChunksColumnName", "Chunks"));
+		Column.SetTitleName(LOCTEXT("ChunksColumnTitle", "Chunks"));
+		Column.SetDescription(LOCTEXT("ChunksColumnDesc", "List of chunks this will be added to when cooked"));
+
+		Column.SetFlags(ETableColumnFlags::ShouldBeVisible | ETableColumnFlags::CanBeHidden | ETableColumnFlags::CanBeFiltered);
+
+		Column.SetHorizontalAlignment(HAlign_Left);
+		Column.SetInitialWidth(120.0f);
+
+		Column.SetDataType(ETableCellDataType::CString);
+
+		class FChunksValueGetter : public FTableCellValueGetter
+		{
+		public:
+			virtual const TOptional<FTableCellValue> GetValue(const FTableColumn& Column, const FBaseTreeNode& Node) const
+			{
+				if (Node.IsGroup())
+				{
+					const FTableTreeNode& NodePtr = static_cast<const FTableTreeNode&>(Node);
+					if (NodePtr.HasAggregatedValue(Column.GetId()))
+					{
+						return NodePtr.GetAggregatedValue(Column.GetId());
+					}
+				}
+				else //if (Node->Is<FAssetTreeNode>())
+				{
+					const FAssetTreeNode& TreeNode = static_cast<const FAssetTreeNode&>(Node);
+					const FAssetTableRow& Asset = TreeNode.GetAssetChecked();
+					return FTableCellValue(Asset.GetChunks());
+				}
+
+				return TOptional<FTableCellValue>();
+			}
+		};
+		TSharedRef<ITableCellValueGetter> Getter = MakeShared<FChunksValueGetter>();
+		Column.SetValueGetter(Getter);
+
+		TSharedRef<ITableCellValueFormatter> Formatter = MakeShared<FCStringValueFormatterAsText>();
+		Column.SetValueFormatter(Formatter);
+
+		TSharedRef<ITableCellValueSorter> Sorter = MakeShared<FSorterByCStringValue>(ColumnRef);
+		Column.SetValueSorter(Sorter);
+
+		Column.SetAggregation(ETableColumnAggregation::SameValue);
+
+		AddColumn(ColumnRef);
+	}
+	//////////////////////////////////////////////////
+	// Native Class Column
+	{
+		TSharedRef<FTableColumn> ColumnRef = MakeShared<FTableColumn>(FAssetTableColumns::NativeClassColumnId);
+		FTableColumn& Column = *ColumnRef;
+
+		Column.SetIndex(ColumnIndex++);
+
+		Column.SetShortName(LOCTEXT("NativeClassColumnName", "Native Class"));
+		Column.SetTitleName(LOCTEXT("NativeClassColumnTitle", "Native Class"));
+		Column.SetDescription(LOCTEXT("NativeClassColumnDesc", "Native class of the asset"));
+
+		Column.SetFlags(ETableColumnFlags::ShouldBeVisible | ETableColumnFlags::CanBeHidden | ETableColumnFlags::CanBeFiltered);
+
+		Column.SetHorizontalAlignment(HAlign_Left);
+		Column.SetInitialWidth(120.0f);
+
+		Column.SetDataType(ETableCellDataType::CString);
+
+		class FNativeClassValueGetter : public FTableCellValueGetter
+		{
+		public:
+			virtual const TOptional<FTableCellValue> GetValue(const FTableColumn& Column, const FBaseTreeNode& Node) const
+			{
+				if (Node.IsGroup())
+				{
+					const FTableTreeNode& NodePtr = static_cast<const FTableTreeNode&>(Node);
+					if (NodePtr.HasAggregatedValue(Column.GetId()))
+					{
+						return NodePtr.GetAggregatedValue(Column.GetId());
+					}
+				}
+				else //if (Node->Is<FAssetTreeNode>())
+				{
+					const FAssetTreeNode& TreeNode = static_cast<const FAssetTreeNode&>(Node);
+					const FAssetTableRow& Asset = TreeNode.GetAssetChecked();
+					return FTableCellValue(Asset.GetNativeClass());
+				}
+
+				return TOptional<FTableCellValue>();
+			}
+		};
+		TSharedRef<ITableCellValueGetter> Getter = MakeShared<FNativeClassValueGetter>();
+		Column.SetValueGetter(Getter);
+
+		TSharedRef<ITableCellValueFormatter> Formatter = MakeShared<FCStringValueFormatterAsText>();
+		Column.SetValueFormatter(Formatter);
+
+		TSharedRef<ITableCellValueSorter> Sorter = MakeShared<FSorterByCStringValue>(ColumnRef);
+		Column.SetValueSorter(Sorter);
+
+		Column.SetAggregation(ETableColumnAggregation::SameValue);
+
+		AddColumn(ColumnRef);
+	}
+	//////////////////////////////////////////////////
+}
+
+////////////////////////////////////////////////////////////////////////////////////////////////////
+
+#undef LOCTEXT_NAMESPACE
diff --git a/Engine/Plugins/Editor/AssetManagerEditor/Source/AssetManagerEditor/Private/TreeView/AssetTable.h b/Engine/Plugins/Editor/AssetManagerEditor/Source/AssetManagerEditor/Private/TreeView/AssetTable.h
new file mode 100644
--- /dev/null
+++ b/Engine/Plugins/Editor/AssetManagerEditor/Source/AssetManagerEditor/Private/TreeView/AssetTable.h
@@ -0,0 +1,95 @@
+// Copyright Epic Games, Inc. All Rights Reserved.
+
+#pragma once
+
+#include "CoreMinimal.h"
+#include "Insights/Table/ViewModels/Table.h"
+
+namespace UE
+{
+	namespace Insights
+	{
+		class FTableColumn;
+	}
+}
+
+// Column identifiers
+struct FAssetTableColumns
+{
+	static const FName CountColumnId;
+	static const FName NameColumnId;
+	static const FName TypeColumnId;
+	static const FName PathColumnId;
+	static const FName PrimaryTypeColumnId;
+	static const FName PrimaryNameColumnId;
+	static const FName ManagedDiskSizeColumnId;
+	static const FName DiskSizeColumnId;
+	static const FName StagedCompressedSizeColumnId;
+	static const FName TotalUsageCountColumnId;
+	static const FName CookRuleColumnId;
+	static const FName ChunksColumnId;
+	static const FName NativeClassColumnId;
+};
+
+class FAssetTableRow
+{
+	friend class FAssetTable;
+	friend class SAssetAuditBrowser;
+
+public:
+	FAssetTableRow()
+	{
+	}
+	~FAssetTableRow()
+	{
+	}
+
+	const TCHAR* GetType() const { return *Type; }
+	const TCHAR* GetName() const { return *Name; }
+	const TCHAR* GetPath() const { return *Path; }
+	const TCHAR* GetPrimaryType() const { return *PrimaryType; }
+	const TCHAR* GetPrimaryName() const { return *PrimaryName; }
+	int64 GetManagedDiskSize() const { return ManagedDiskSize; }
+	int64 GetDiskSize() const { return DiskSize; }
+	int64 GetStagedCompressedSize() const { return StagedCompressedSize; }
+	int64 GetTotalUsageCount() const { return TotalUsageCount; }
+	const TCHAR* GetCookRule() const { return *CookRule; }
+	const TCHAR* GetChunks() const { return *Chunks; }
+	const TCHAR* GetNativeClass() const { return *NativeClass; }
+
+private:
+	FString Type;
+	FString Name;
+	FString Path;
+	FString PrimaryType;
+	FString PrimaryName;
+	int64 ManagedDiskSize = 0;
+	int64 DiskSize = 0;
+	int64 StagedCompressedSize = 0;
+	int64 TotalUsageCount = 0;
+	FString CookRule;
+	FString Chunks;
+	FString NativeClass;
+};
+
+class FAssetTable : public UE::Insights::FTable
+{
+public:
+	FAssetTable();
+	virtual ~FAssetTable();
+
+	virtual void Reset();
+
+	TArray<FAssetTableRow>& GetAssets() { return Assets; }
+	const TArray<FAssetTableRow>& GetAssets() const { return Assets; }
+
+	bool IsValidRowIndex(int32 InIndex) const { return InIndex >= 0 && InIndex < Assets.Num(); }
+	const FAssetTableRow* GetAsset(int32 InIndex) const { return IsValidRowIndex(InIndex) ? &Assets[InIndex] : nullptr; }
+	const FAssetTableRow& GetAssetChecked(int32 InIndex) const { check(IsValidRowIndex(InIndex)); return Assets[InIndex]; }
+
+private:
+	void AddDefaultColumns();
+
+private:
+	TArray<FAssetTableRow> Assets;
+};
diff --git a/Engine/Plugins/Editor/AssetManagerEditor/Source/AssetManagerEditor/Private/TreeView/AssetTreeNode.cpp b/Engine/Plugins/Editor/AssetManagerEditor/Source/AssetManagerEditor/Private/TreeView/AssetTreeNode.cpp
new file mode 100644
--- /dev/null
+++ b/Engine/Plugins/Editor/AssetManagerEditor/Source/AssetManagerEditor/Private/TreeView/AssetTreeNode.cpp
@@ -0,0 +1,5 @@
+// Copyright Epic Games, Inc. All Rights Reserved.
+
+#include "AssetTreeNode.h"
+
+INSIGHTS_IMPLEMENT_RTTI(FAssetTreeNode)
diff --git a/Engine/Plugins/Editor/AssetManagerEditor/Source/AssetManagerEditor/Private/TreeView/AssetTreeNode.h b/Engine/Plugins/Editor/AssetManagerEditor/Source/AssetManagerEditor/Private/TreeView/AssetTreeNode.h
new file mode 100644
--- /dev/null
+++ b/Engine/Plugins/Editor/AssetManagerEditor/Source/AssetManagerEditor/Private/TreeView/AssetTreeNode.h
@@ -0,0 +1,46 @@
+// Copyright Epic Games, Inc. All Rights Reserved.
+
+#pragma once
+
+#include "CoreMinimal.h"
+
+#include "AssetTable.h"
+#include "Insights/Table/ViewModels/TableTreeNode.h"
+
+class FAssetTreeNode : public UE::Insights::FTableTreeNode
+{
+	INSIGHTS_DECLARE_RTTI(FAssetTreeNode, UE::Insights::FTableTreeNode)
+
+public:
+	/** Initialization constructor for the asset node. */
+	explicit FAssetTreeNode(const FName InName, TWeakPtr<FAssetTable> InParentTable, int32 InRowIndex)
+		: FTableTreeNode(InName, InParentTable, InRowIndex)
+		, AssetTable(InParentTable.Pin().Get())
+	{
+	}
+
+	/** Initialization constructor for the group node. */
+	explicit FAssetTreeNode(const FName InGroupName, TWeakPtr<FAssetTable> InParentTable)
+		: FTableTreeNode(InGroupName, InParentTable)
+		, AssetTable(InParentTable.Pin().Get())
+	{
+	}
+
+	FAssetTable& GetAssetTableChecked() const { return *AssetTable; }
+	const FAssetTableRow& GetAssetChecked() const { return AssetTable->GetAssetChecked(RowId.RowIndex); }
+
+private:
+	FAssetTable* AssetTable;
+};
+
+/** Type definition for shared pointers to instances of FAssetTreeNode. */
+typedef TSharedPtr<class FAssetTreeNode> FAssetTreeNodePtr;
+
+/** Type definition for shared references to instances of FAssetTreeNode. */
+typedef TSharedRef<class FAssetTreeNode> FAssetTreeNodeRef;
+
+/** Type definition for shared references to const instances of FAssetTreeNode. */
+typedef TSharedRef<const class FAssetTreeNode> FAssetTreeNodeRefConst;
+
+/** Type definition for weak references to instances of FAssetTreeNode. */
+typedef TWeakPtr<class FAssetTreeNode> FAssetTreeNodeWeak;
diff --git a/Engine/Plugins/Editor/AssetManagerEditor/Source/AssetManagerEditor/Private/TreeView/SAssetTableTreeView.cpp b/Engine/Plugins/Editor/AssetManagerEditor/Source/AssetManagerEditor/Private/TreeView/SAssetTableTreeView.cpp
new file mode 100644
--- /dev/null
+++ b/Engine/Plugins/Editor/AssetManagerEditor/Source/AssetManagerEditor/Private/TreeView/SAssetTableTreeView.cpp
@@ -0,0 +1,503 @@
+// Copyright Epic Games, Inc. All Rights Reserved.
+
+#include "SAssetTableTreeView.h"
+
+#include "Containers/Set.h"
+#include "Framework/MultiBox/MultiBoxBuilder.h"
+#include "GenericPlatform/GenericPlatformMisc.h"
+#include "Modules/ModuleManager.h"
+#include "SlateOptMacros.h"
+#include "Styling/AppStyle.h"
+#include "Widgets/Images/SImage.h"
+#include "Widgets/Input/SCheckBox.h"
+#include "Widgets/SToolTip.h"
+#include "DesktopPlatformModule.h"
+#include "HAL/PlatformFileManager.h"
+#include "Logging/MessageLog.h"
+
+#include "Insights/Common/Log.h"
+#include "Insights/Common/Stopwatch.h"
+#include "Insights/Filter/ViewModels/FilterConfigurator.h"
+#include "Insights/Table/ViewModels/TableCellValueFormatter.h"
+#include "Insights/Table/ViewModels/TableColumn.h"
+#include "Insights/Table/ViewModels/TreeNodeGrouping.h"
+#include "TreeView/AssetTable.h"
+#include "TreeView/AssetTreeNode.h"
+
+#include <limits>
+#include <memory>
+
+#define LOCTEXT_NAMESPACE "SAssetTableTreeView"
+
+////////////////////////////////////////////////////////////////////////////////////////////////////
+
+SAssetTableTreeView::SAssetTableTreeView()
+{
+	bRunInAsyncMode = true;
+}
+
+////////////////////////////////////////////////////////////////////////////////////////////////////
+
+SAssetTableTreeView::~SAssetTableTreeView()
+{
+}
+
+////////////////////////////////////////////////////////////////////////////////////////////////////
+
+void SAssetTableTreeView::Construct(const FArguments& InArgs, TSharedPtr<FAssetTable> InTablePtr)
+{
+	STableTreeView::ConstructWidget(InTablePtr);
+
+	CreateGroupings();
+	CreateSortings();
+}
+
+////////////////////////////////////////////////////////////////////////////////////////////////////
+
+void SAssetTableTreeView::Reset()
+{
+	//...
+	UE::Insights::STableTreeView::Reset();
+}
+
+////////////////////////////////////////////////////////////////////////////////////////////////////
+
+void SAssetTableTreeView::Tick(const FGeometry& AllottedGeometry, const double InCurrentTime, const float InDeltaTime)
+{
+	UE::Insights::STableTreeView::Tick(AllottedGeometry, InCurrentTime, InDeltaTime);
+
+	if (bNeedsToRebuild && !bIsUpdateRunning)
+	{
+		RebuildTree(true);
+		bNeedsToRebuild = false;
+	}
+}
+
+////////////////////////////////////////////////////////////////////////////////////////////////////
+
+void SAssetTableTreeView::RebuildTree(bool bResync)
+{
+	UE::Insights::FStopwatch Stopwatch;
+	Stopwatch.Start();
+
+	UE::Insights::FStopwatch SyncStopwatch;
+	SyncStopwatch.Start();
+
+	if (bResync)
+	{
+		TableTreeNodes.Empty();
+	}
+
+	const int32 PreviousNodeCount = TableTreeNodes.Num();
+
+	TSharedPtr<FAssetTable> AssetTable = GetAssetTable();
+	const TArray<FAssetTableRow>& Assets = AssetTable->GetAssets();
+	const int32 TotalAssetCount = Assets.Num();
+
+	if (TotalAssetCount != PreviousNodeCount)
+	{
+		UE_LOG(LogInsights, Log, TEXT("[AssetTree] Creating nodes (%d nodes --> %d assets)..."), PreviousNodeCount, TotalAssetCount);
+
+		if (TotalAssetCount < PreviousNodeCount)
+		{
+			TableTreeNodes.Empty();
+		}
+		TableTreeNodes.Reserve(TotalAssetCount);
+
+		//const FName BaseNodeName(TEXT("asset"));
+		for (int32 AssetIndex = TableTreeNodes.Num(); AssetIndex < TotalAssetCount; ++AssetIndex)
+		{
+			const FAssetTableRow* Asset = AssetTable->GetAsset(AssetIndex);
+
+			//FName NodeName(BaseNodeName, AssetIndex + 1);
+			FName NodeName(Asset->GetName());
+			FAssetTreeNodePtr NodePtr = MakeShared<FAssetTreeNode>(NodeName, AssetTable, AssetIndex);
+			TableTreeNodes.Add(NodePtr);
+		}
+		ensure(TableTreeNodes.Num() == TotalAssetCount);
+	}
+
+	SyncStopwatch.Stop();
+
+	if (bResync || TableTreeNodes.Num() != PreviousNodeCount)
+	{
+		// Save selection.
+		TArray<UE::Insights::FTableTreeNodePtr> SelectedItems;
+		TreeView->GetSelectedItems(SelectedItems);
+
+		UE_LOG(LogInsights, Log, TEXT("[AssetTree] Update tree..."));
+		UpdateTree();
+
+		TreeView->RebuildList();
+
+		// Restore selection.
+		if (SelectedItems.Num() > 0)
+		{
+			TreeView->ClearSelection();
+			for (UE::Insights::FTableTreeNodePtr& NodePtr : SelectedItems)
+			{
+				NodePtr = GetNodeByTableRowIndex(NodePtr->GetRowIndex());
+			}
+			SelectedItems.RemoveAll([](const UE::Insights::FTableTreeNodePtr& NodePtr) { return !NodePtr.IsValid(); });
+			if (SelectedItems.Num() > 0)
+			{
+				TreeView->SetItemSelection(SelectedItems, true);
+				TreeView->RequestScrollIntoView(SelectedItems.Last());
+			}
+		}
+	}
+
+	Stopwatch.Stop();
+	const double TotalTime = Stopwatch.GetAccumulatedTime();
+	if (TotalTime > 0.01)
+	{
+		const double SyncTime = SyncStopwatch.GetAccumulatedTime();
+		UE_LOG(LogInsights, Log, TEXT("[AssetTree] Tree view rebuilt in %.4fs (sync: %.4fs + update: %.4fs) --> %d nodes (%d added)"),
+			TotalTime, SyncTime, TotalTime - SyncTime, TableTreeNodes.Num(), TableTreeNodes.Num() - PreviousNodeCount);
+	}
+}
+
+////////////////////////////////////////////////////////////////////////////////////////////////////
+
+bool SAssetTableTreeView::IsRunning() const
+{
+	return STableTreeView::IsRunning();
+}
+
+////////////////////////////////////////////////////////////////////////////////////////////////////
+
+double SAssetTableTreeView::GetAllOperationsDuration()
+{
+	return STableTreeView::GetAllOperationsDuration();
+}
+
+////////////////////////////////////////////////////////////////////////////////////////////////////
+
+FText SAssetTableTreeView::GetCurrentOperationName() const
+{
+	return STableTreeView::GetCurrentOperationName();
+}
+
+////////////////////////////////////////////////////////////////////////////////////////////////////
+
+BEGIN_SLATE_FUNCTION_BUILD_OPTIMIZATION
+TSharedPtr<SWidget> SAssetTableTreeView::ConstructToolbar()
+{
+	TSharedPtr<SHorizontalBox> Box = SNew(SHorizontalBox);
+
+	Box->AddSlot()
+		.AutoWidth()
+		.VAlign(VAlign_Center)
+		.Padding(0.0f, 0.0f, 4.0f, 0.0f)
+		[
+			SNew(STextBlock)
+			.Text(LOCTEXT("Preset", "Preset:"))
+		];
+
+	Box->AddSlot()
+		.AutoWidth()
+		.VAlign(VAlign_Center)
+		[
+			SNew(SBox)
+			.MinDesiredWidth(150.0f)
+			[
+				SAssignNew(PresetComboBox, SComboBox<TSharedRef<UE::Insights::ITableTreeViewPreset>>)
+				.ToolTipText(this, &SAssetTableTreeView::ViewPreset_GetSelectedToolTipText)
+				.OptionsSource(GetAvailableViewPresets())
+				.OnSelectionChanged(this, &SAssetTableTreeView::ViewPreset_OnSelectionChanged)
+				.OnGenerateWidget(this, &SAssetTableTreeView::ViewPreset_OnGenerateWidget)
+				[
+					SNew(STextBlock)
+					.Text(this, &SAssetTableTreeView::ViewPreset_GetSelectedText)
+				]
+			]
+		];
+
+	return Box;
+}
+END_SLATE_FUNCTION_BUILD_OPTIMIZATION
+
+////////////////////////////////////////////////////////////////////////////////////////////////////
+
+void SAssetTableTreeView::InitAvailableViewPresets()
+{
+	//////////////////////////////////////////////////
+	// Default View
+
+	class FDefaultViewPreset : public UE::Insights::ITableTreeViewPreset
+	{
+	public:
+		virtual FText GetName() const override
+		{
+			return LOCTEXT("Default_PresetName", "Default");
+		}
+		virtual FText GetToolTip() const override
+		{
+			return LOCTEXT("Default_PresetToolTip", "Default View\nConfigure the tree view to show default asset info.");
+		}
+		virtual FName GetSortColumn() const override
+		{
+			return UE::Insights::FTable::GetHierarchyColumnId();
+		}
+		virtual EColumnSortMode::Type GetSortMode() const override
+		{
+			return EColumnSortMode::Type::Ascending;
+		}
+		virtual void SetCurrentGroupings(const TArray<TSharedPtr<UE::Insights::FTreeNodeGrouping>>& InAvailableGroupings, TArray<TSharedPtr<UE::Insights::FTreeNodeGrouping>>& InOutCurrentGroupings) const override
+		{
+			InOutCurrentGroupings.Reset();
+
+			check(InAvailableGroupings[0]->Is<UE::Insights::FTreeNodeGroupingFlat>());
+			InOutCurrentGroupings.Add(InAvailableGroupings[0]);
+		}
+		virtual void GetColumnConfigSet(TArray<UE::Insights::FTableColumnConfig>& InOutConfigSet) const override
+		{
+			InOutConfigSet.Add({ UE::Insights::FTable::GetHierarchyColumnId(),         true, 200.0f });
+			InOutConfigSet.Add({ FAssetTableColumns::CountColumnId,                true, 100.0f });
+			InOutConfigSet.Add({ FAssetTableColumns::NameColumnId,                 true, 200.0f });
+			InOutConfigSet.Add({ FAssetTableColumns::TypeColumnId,                 true, 200.0f });
+			InOutConfigSet.Add({ FAssetTableColumns::PathColumnId,                 true, 400.0f });
+			InOutConfigSet.Add({ FAssetTableColumns::PrimaryTypeColumnId,          true, 200.0f });
+			InOutConfigSet.Add({ FAssetTableColumns::PrimaryNameColumnId,          true, 200.0f });
+			InOutConfigSet.Add({ FAssetTableColumns::ManagedDiskSizeColumnId,      true, 100.0f });
+			InOutConfigSet.Add({ FAssetTableColumns::DiskSizeColumnId,             true, 100.0f });
+			InOutConfigSet.Add({ FAssetTableColumns::StagedCompressedSizeColumnId, true, 100.0f });
+			InOutConfigSet.Add({ FAssetTableColumns::TotalUsageCountColumnId,      true, 100.0f });
+			InOutConfigSet.Add({ FAssetTableColumns::CookRuleColumnId,             true, 200.0f });
+			InOutConfigSet.Add({ FAssetTableColumns::ChunksColumnId,               true, 200.0f });
+			InOutConfigSet.Add({ FAssetTableColumns::NativeClassColumnId,          true, 200.0f });
+		}
+	};
+	AvailableViewPresets.Add(MakeShared<FDefaultViewPreset>());
+
+	//////////////////////////////////////////////////
+	// Path Breakdown View
+
+	class FAssetPathViewPreset : public UE::Insights::ITableTreeViewPreset
+	{
+	public:
+		virtual FText GetName() const override
+		{
+			return LOCTEXT("Path_PresetName", "Path");
+		}
+		virtual FText GetToolTip() const override
+		{
+			return LOCTEXT("Path_PresetToolTip", "Path Breakdown View\nConfigure the tree view to show a breakdown of assets by their path.");
+		}
+		virtual FName GetSortColumn() const override
+		{
+			return UE::Insights::FTable::GetHierarchyColumnId();
+		}
+		virtual EColumnSortMode::Type GetSortMode() const override
+		{
+			return EColumnSortMode::Type::Ascending;
+		}
+		virtual void SetCurrentGroupings(const TArray<TSharedPtr<UE::Insights::FTreeNodeGrouping>>& InAvailableGroupings, TArray<TSharedPtr<UE::Insights::FTreeNodeGrouping>>& InOutCurrentGroupings) const override
+		{
+			InOutCurrentGroupings.Reset();
+
+			check(InAvailableGroupings[0]->Is<UE::Insights::FTreeNodeGroupingFlat>());
+			InOutCurrentGroupings.Add(InAvailableGroupings[0]);
+
+			const TSharedPtr<UE::Insights::FTreeNodeGrouping>* PathGrouping = InAvailableGroupings.FindByPredicate(
+				[](TSharedPtr<UE::Insights::FTreeNodeGrouping>& Grouping)
+				{
+					return Grouping->Is<UE::Insights::FTreeNodeGroupingByPathBreakdown>() &&
+						   Grouping->As<UE::Insights::FTreeNodeGroupingByPathBreakdown>().GetColumnId() == FAssetTableColumns::PathColumnId;
+				});
+			if (PathGrouping)
+			{
+				InOutCurrentGroupings.Add(*PathGrouping);
+			}
+		}
+		virtual void GetColumnConfigSet(TArray<UE::Insights::FTableColumnConfig>& InOutConfigSet) const override
+		{
+			InOutConfigSet.Add({ UE::Insights::FTable::GetHierarchyColumnId(),         true, 400.0f });
+			InOutConfigSet.Add({ FAssetTableColumns::CountColumnId,                true, 100.0f });
+			InOutConfigSet.Add({ FAssetTableColumns::NameColumnId,                 true, 200.0f });
+			InOutConfigSet.Add({ FAssetTableColumns::TypeColumnId,                 true, 200.0f });
+			//InOutConfigSet.Add({ FAssetTableColumns::PathColumnId,                 true, 400.0f });
+			InOutConfigSet.Add({ FAssetTableColumns::PrimaryTypeColumnId,          true, 200.0f });
+			InOutConfigSet.Add({ FAssetTableColumns::PrimaryNameColumnId,          true, 200.0f });
+			InOutConfigSet.Add({ FAssetTableColumns::ManagedDiskSizeColumnId,      true, 100.0f });
+			InOutConfigSet.Add({ FAssetTableColumns::DiskSizeColumnId,             true, 100.0f });
+			InOutConfigSet.Add({ FAssetTableColumns::StagedCompressedSizeColumnId, true, 100.0f });
+			InOutConfigSet.Add({ FAssetTableColumns::TotalUsageCountColumnId,      true, 100.0f });
+			InOutConfigSet.Add({ FAssetTableColumns::CookRuleColumnId,             true, 200.0f });
+			InOutConfigSet.Add({ FAssetTableColumns::ChunksColumnId,               true, 200.0f });
+			InOutConfigSet.Add({ FAssetTableColumns::NativeClassColumnId,          true, 200.0f });
+		}
+	};
+	AvailableViewPresets.Add(MakeShared<FAssetPathViewPreset>());
+
+	//////////////////////////////////////////////////
+	// Primary Asset Breakdown View
+
+	class FPrimaryAssetViewPreset : public UE::Insights::ITableTreeViewPreset
+	{
+	public:
+		virtual FText GetName() const override
+		{
+			return LOCTEXT("PrimaryAsset_PresetName", "Primary Asset");
+		}
+		virtual FText GetToolTip() const override
+		{
+			return LOCTEXT("PrimaryAsset_PresetToolTip", "Primary Asset Breakdown View\nConfigure the tree view to show a breakdown of assets by their primary asset type/name.");
+		}
+		virtual FName GetSortColumn() const override
+		{
+			return UE::Insights::FTable::GetHierarchyColumnId();
+		}
+		virtual EColumnSortMode::Type GetSortMode() const override
+		{
+			return EColumnSortMode::Type::Ascending;
+		}
+		virtual void SetCurrentGroupings(const TArray<TSharedPtr<UE::Insights::FTreeNodeGrouping>>& InAvailableGroupings, TArray<TSharedPtr<UE::Insights::FTreeNodeGrouping>>& InOutCurrentGroupings) const override
+		{
+			InOutCurrentGroupings.Reset();
+
+			check(InAvailableGroupings[0]->Is<UE::Insights::FTreeNodeGroupingFlat>());
+			InOutCurrentGroupings.Add(InAvailableGroupings[0]);
+
+			const TSharedPtr<UE::Insights::FTreeNodeGrouping>* PrimaryTypeGrouping = InAvailableGroupings.FindByPredicate(
+				[](TSharedPtr<UE::Insights::FTreeNodeGrouping>& Grouping)
+				{
+					return Grouping->Is<UE::Insights::FTreeNodeGroupingByUniqueValueCString>() &&
+						   Grouping->As<UE::Insights::FTreeNodeGroupingByUniqueValueCString>().GetColumnId() == FAssetTableColumns::PrimaryTypeColumnId;
+				});
+			if (PrimaryTypeGrouping)
+			{
+				InOutCurrentGroupings.Add(*PrimaryTypeGrouping);
+			}
+
+			const TSharedPtr<UE::Insights::FTreeNodeGrouping>* PrimaryNameGrouping = InAvailableGroupings.FindByPredicate(
+				[](TSharedPtr<UE::Insights::FTreeNodeGrouping>& Grouping)
+				{
+					return Grouping->Is<UE::Insights::FTreeNodeGroupingByUniqueValueCString>() &&
+						   Grouping->As<UE::Insights::FTreeNodeGroupingByUniqueValueCString>().GetColumnId() == FAssetTableColumns::PrimaryNameColumnId;
+				});
+			if (PrimaryNameGrouping)
+			{
+				InOutCurrentGroupings.Add(*PrimaryNameGrouping);
+			}
+		}
+		virtual void GetColumnConfigSet(TArray<UE::Insights::FTableColumnConfig>& InOutConfigSet) const override
+		{
+			InOutConfigSet.Add({ UE::Insights::FTable::GetHierarchyColumnId(),         true, 300.0f });
+			InOutConfigSet.Add({ FAssetTableColumns::CountColumnId,                true, 100.0f });
+			InOutConfigSet.Add({ FAssetTableColumns::NameColumnId,                 true, 200.0f });
+			InOutConfigSet.Add({ FAssetTableColumns::TypeColumnId,                 true, 200.0f });
+			InOutConfigSet.Add({ FAssetTableColumns::PathColumnId,                 true, 400.0f });
+			//InOutConfigSet.Add({ FAssetTableColumns::PrimaryTypeColumnId,          true, 200.0f });
+			//InOutConfigSet.Add({ FAssetTableColumns::PrimaryNameColumnId,          true, 200.0f });
+			InOutConfigSet.Add({ FAssetTableColumns::ManagedDiskSizeColumnId,      true, 100.0f });
+			InOutConfigSet.Add({ FAssetTableColumns::DiskSizeColumnId,             true, 100.0f });
+			InOutConfigSet.Add({ FAssetTableColumns::StagedCompressedSizeColumnId, true, 100.0f });
+			InOutConfigSet.Add({ FAssetTableColumns::TotalUsageCountColumnId,      true, 100.0f });
+			InOutConfigSet.Add({ FAssetTableColumns::CookRuleColumnId,             true, 200.0f });
+			InOutConfigSet.Add({ FAssetTableColumns::ChunksColumnId,               true, 200.0f });
+			InOutConfigSet.Add({ FAssetTableColumns::NativeClassColumnId,          true, 200.0f });
+		}
+	};
+	AvailableViewPresets.Add(MakeShared<FPrimaryAssetViewPreset>());
+
+	//////////////////////////////////////////////////
+
+	SelectedViewPreset = AvailableViewPresets[0];
+}
+
+////////////////////////////////////////////////////////////////////////////////////////////////////
+
+BEGIN_SLATE_FUNCTION_BUILD_OPTIMIZATION
+TSharedPtr<SWidget> SAssetTableTreeView::ConstructFooter()
+{
+	return
+		SNew(SHorizontalBox)
+
+		+ SHorizontalBox::Slot()
+		.HAlign(HAlign_Left)
+		.Padding(2.0f)
+		[
+			SNew(STextBlock)
+			.Text(LOCTEXT("TODO", "<TODO:status>"))
+		]
+
+		+ SHorizontalBox::Slot()
+		.HAlign(HAlign_Right)
+		.Padding(2.0f)
+		[
+			SNew(STextBlock)
+			.Text(LOCTEXT("TODO", "<TODO:status>"))
+		];
+}
+END_SLATE_FUNCTION_BUILD_OPTIMIZATION
+
+////////////////////////////////////////////////////////////////////////////////////////////////////
+
+void SAssetTableTreeView::InternalCreateGroupings()
+{
+	UE::Insights::STableTreeView::InternalCreateGroupings();
+
+	AvailableGroupings.RemoveAll(
+		[](TSharedPtr<UE::Insights::FTreeNodeGrouping>& Grouping)
+		{
+			if (Grouping->Is<UE::Insights::FTreeNodeGroupingByUniqueValue>())
+			{
+				const FName ColumnId = Grouping->As<UE::Insights::FTreeNodeGroupingByUniqueValue>().GetColumnId();
+				if (ColumnId == FAssetTableColumns::CountColumnId)
+				{
+					return true;
+				}
+			}
+			else if (Grouping->Is<UE::Insights::FTreeNodeGroupingByPathBreakdown>())
+			{
+				const FName ColumnId = Grouping->As<UE::Insights::FTreeNodeGroupingByPathBreakdown>().GetColumnId();
+				if (ColumnId == FAssetTableColumns::TypeColumnId ||
+					ColumnId == FAssetTableColumns::NativeClassColumnId)
+				{
+					return true;
+				}
+			}
+			return false;
+		});
+
+	int32 Index = 1; // after the Flat ("All") grouping
+}
+
+////////////////////////////////////////////////////////////////////////////////////////////////////
+
+bool SAssetTableTreeView::ApplyCustomAdvancedFilters(const UE::Insights::FTableTreeNodePtr& NodePtr)
+{
+	return true;
+}
+
+////////////////////////////////////////////////////////////////////////////////////////////////////
+
+void SAssetTableTreeView::AddCustomAdvancedFilters()
+{
+}
+
+////////////////////////////////////////////////////////////////////////////////////////////////////
+
+TSharedPtr<FAssetTreeNode> SAssetTableTreeView::GetSingleSelectedAssetNode() const
+{
+	if (TreeView->GetNumItemsSelected() == 1)
+	{
+		FAssetTreeNodePtr SelectedTreeNode = StaticCastSharedPtr<FAssetTreeNode>(TreeView->GetSelectedItems()[0]);
+		if (SelectedTreeNode.IsValid() && !SelectedTreeNode->IsGroup())
+		{
+			return SelectedTreeNode;
+		}
+	}
+	return nullptr;
+}
+
+////////////////////////////////////////////////////////////////////////////////////////////////////
+
+void SAssetTableTreeView::ExtendMenu(FMenuBuilder& MenuBuilder)
+{
+}
+
+////////////////////////////////////////////////////////////////////////////////////////////////////
+
+#undef LOCTEXT_NAMESPACE
diff --git a/Engine/Plugins/Editor/AssetManagerEditor/Source/AssetManagerEditor/Private/TreeView/SAssetTableTreeView.h b/Engine/Plugins/Editor/AssetManagerEditor/Source/AssetManagerEditor/Private/TreeView/SAssetTableTreeView.h
new file mode 100644
--- /dev/null
+++ b/Engine/Plugins/Editor/AssetManagerEditor/Source/AssetManagerEditor/Private/TreeView/SAssetTableTreeView.h
@@ -0,0 +1,81 @@
+// Copyright Epic Games, Inc. All Rights Reserved.
+
+#pragma once
+
+#include "CoreMinimal.h"
+#include "Widgets/DeclarativeSyntaxSupport.h"
+
+#include "TreeView/AssetTable.h"
+#include "Insights/Table/Widgets/STableTreeView.h"
+
+class FAssetTreeNode;
+
+class SAssetTableTreeView : public UE::Insights::STableTreeView
+{
+public:
+	/** Default constructor. */
+	SAssetTableTreeView();
+
+	/** Virtual destructor. */
+	virtual ~SAssetTableTreeView();
+
+	SLATE_BEGIN_ARGS(SAssetTableTreeView) {}
+	SLATE_END_ARGS()
+
+	/**
+	 * Construct this widget
+	 * @param InArgs - The declaration data for this widget
+	 */
+	void Construct(const FArguments& InArgs, TSharedPtr<FAssetTable> InTablePtr);
+
+	virtual TSharedPtr<SWidget> ConstructToolbar() override;
+	virtual TSharedPtr<SWidget> ConstructFooter() override;
+
+	TSharedPtr<FAssetTable> GetAssetTable() { return StaticCastSharedPtr<FAssetTable>(GetTable()); }
+	const TSharedPtr<FAssetTable> GetAssetTable() const { return StaticCastSharedPtr<FAssetTable>(GetTable()); }
+
+	//void UpdateSourceTable(TSharedPtr<TraceServices::IAssetTable> SourceTable);
+
+	virtual void Reset();
+
+	/**
+	 * Ticks this widget.  Override in derived classes, but always call the parent implementation.
+	 *
+	 * @param  AllottedGeometry The space allotted for this widget
+	 * @param  InCurrentTime  Current absolute real time
+	 * @param  InDeltaTime  Real time passed since last tick
+	 */
+	virtual void Tick(const FGeometry& AllottedGeometry, const double InCurrentTime, const float InDeltaTime) override;
+
+	/**
+	 * Rebuilds the tree (if necessary).
+	 * @param bResync - If true, it forces a resync even if the list did not changed since last sync.
+	 */
+	virtual void RebuildTree(bool bResync);
+	virtual void RebuildTreeAsync() { bNeedsToRebuild = true; }
+
+	////////////////////////////////////////////////////////////////////////////////////////////////////
+	// IAsyncOperationStatusProvider implementation
+
+	virtual bool IsRunning() const override;
+	virtual double GetAllOperationsDuration() override;
+	virtual FText GetCurrentOperationName() const override;
+
+	////////////////////////////////////////////////////////////////////////////////////////////////////
+
+protected:
+	virtual void InternalCreateGroupings() override;
+
+	TSharedPtr<FAssetTreeNode> GetSingleSelectedAssetNode() const;
+
+	virtual void ExtendMenu(FMenuBuilder& MenuBuilder) override;
+
+private:
+	bool virtual ApplyCustomAdvancedFilters(const UE::Insights::FTableTreeNodePtr& NodePtr) override;
+	virtual void AddCustomAdvancedFilters() override;
+
+	void InitAvailableViewPresets();
+
+private:
+	bool bNeedsToRebuild = false;
+};
