Further WIP additions to the asset audit browser for exploring disk size and dependency data. Changes include modification of columns, additional support for tint and icons (incomplete), dependency nodes, and asset dependency identification (incomplete). Code is being submitted disabled via ifdef.

diff --git a/Engine/Plugins/Editor/AssetManagerEditor/Source/AssetManagerEditor/Private/AssetManagerEditorModule.cpp b/Engine/Plugins/Editor/AssetManagerEditor/Source/AssetManagerEditor/Private/AssetManagerEditorModule.cpp
--- a/Engine/Plugins/Editor/AssetManagerEditor/Source/AssetManagerEditor/Private/AssetManagerEditorModule.cpp
+++ b/Engine/Plugins/Editor/AssetManagerEditor/Source/AssetManagerEditor/Private/AssetManagerEditorModule.cpp
@@ -111,6 +111,7 @@ const FName IAssetManagerEditorModule::CookRuleName = FName("CookRule");
 const FName IAssetManagerEditorModule::ChunksName = FName("Chunks");
 const FName IAssetManagerEditorModule::StageChunkSizeName = FName("Stage_ChunkSize");
 const FName IAssetManagerEditorModule::StageChunkCompressedSizeName = FName("Stage_ChunkCompressedSize");
+const FName IAssetManagerEditorModule::GameFeaturePluginsName = FName("GameFeaturePlugins");
 
 const FString FAssetManagerEditorRegistrySource::EditorSourceName = TEXT("Editor");
 const FString FAssetManagerEditorRegistrySource::CustomSourceName = TEXT("Custom");
@@ -1300,6 +1301,19 @@ bool FAssetManagerEditorModule::GetStringValueForCustomColumn(const FAssetData&
 		}
 		return true;
 	}
+	else if (ColumnName == GameFeaturePluginsName)
+	{
+		TArray<FString> ParsedPath;
+		const FString& PackageNameString = AssetData.PackageName.ToString();
+		if (PackageNameString.StartsWith(TEXT("/Game/")))
+		{
+			OutValue = TEXT("BaseGame");
+		}
+		else if (PackageNameString.ParseIntoArray(ParsedPath, TEXT("/")) >= 1)
+		{
+			OutValue = ParsedPath[0];
+		}
+	}
 	else
 	{
 		// Get base value of asset tag
@@ -1361,7 +1375,7 @@ bool FAssetManagerEditorModule::GetDisplayTextForCustomColumn(const FAssetData&
 			return true;
 		}
 	}
-	else if (ColumnName == ChunksName)
+	else if (ColumnName == ChunksName || ColumnName == GameFeaturePluginsName)
 	{
 		FString OutString;
 
diff --git a/Engine/Plugins/Editor/AssetManagerEditor/Source/AssetManagerEditor/Private/Insights/Common/InsightsStyle.cpp b/Engine/Plugins/Editor/AssetManagerEditor/Source/AssetManagerEditor/Private/Insights/Common/InsightsStyle.cpp
--- a/Engine/Plugins/Editor/AssetManagerEditor/Source/AssetManagerEditor/Private/Insights/Common/InsightsStyle.cpp
+++ b/Engine/Plugins/Editor/AssetManagerEditor/Source/AssetManagerEditor/Private/Insights/Common/InsightsStyle.cpp
@@ -162,6 +162,9 @@ void FInsightsStyle::FStyle::Initialize()
 
 	Set("Icons.Hint.TreeItem", new IMAGE_BRUSH_SVG("InfoTag_12", Icon12x12));
 	Set("Icons.Group.TreeItem", new CORE_IMAGE_BRUSH_SVG("Starship/Common/folder-closed", Icon12x12));
+	Set("Icons.Leaf.TreeItem", new CORE_IMAGE_BRUSH_SVG("Starship/Common/bullet-point", Icon12x12));
+	Set("Icons.Asset.TreeItem", new CORE_IMAGE_BRUSH_SVG("Starship/Common/box-perspective", Icon12x12));
+	Set("Icons.Package.TreeItem", new EDITOR_IMAGE_BRUSH_SVG("Starship/Common/ProjectPackage", Icon12x12));
 
 	Set("TreeTable.TooltipBold", FTextBlockStyle(NormalText)
 		.SetFont(DEFAULT_FONT("Bold", 8))
diff --git a/Engine/Plugins/Editor/AssetManagerEditor/Source/AssetManagerEditor/Private/Insights/Table/ViewModels/BaseTreeNode.cpp b/Engine/Plugins/Editor/AssetManagerEditor/Source/AssetManagerEditor/Private/Insights/Table/ViewModels/BaseTreeNode.cpp
--- a/Engine/Plugins/Editor/AssetManagerEditor/Source/AssetManagerEditor/Private/Insights/Table/ViewModels/BaseTreeNode.cpp
+++ b/Engine/Plugins/Editor/AssetManagerEditor/Source/AssetManagerEditor/Private/Insights/Table/ViewModels/BaseTreeNode.cpp
@@ -3,6 +3,7 @@
 #include "BaseTreeNode.h"
 
 // Insights
+#include "Insights/Common/InsightsStyle.h"
 #include "Insights/Table/ViewModels/TableCellValueSorter.h"
 
 #define LOCTEXT_NAMESPACE "UE::Insights::FBaseTreeNode"
@@ -54,6 +55,27 @@ bool FBaseTreeNode::HasExtraDisplayName() const
 
 ////////////////////////////////////////////////////////////////////////////////////////////////////
 
+void FBaseTreeNode::SetDefaultIcon()
+{
+	if (IsGroup())
+	{
+		IconBrush = FInsightsStyle::GetBrush("Icons.Group.TreeItem");
+	}
+	else
+	{
+		IconBrush = FInsightsStyle::GetBrush("Icons.Leaf.TreeItem");
+	}
+}
+
+////////////////////////////////////////////////////////////////////////////////////////////////////
+
+FLinearColor FBaseTreeNode::GetColor() const
+{
+	return FLinearColor(1.0f, 1.0f, 1.0f, 1.0f);
+}
+
+////////////////////////////////////////////////////////////////////////////////////////////////////
+
 void FBaseTreeNode::SortChildrenAscending(const ITableCellValueSorter& Sorter)
 {
 	Sorter.Sort(GroupData->Children, ESortMode::Ascending);
diff --git a/Engine/Plugins/Editor/AssetManagerEditor/Source/AssetManagerEditor/Private/Insights/Table/ViewModels/BaseTreeNode.h b/Engine/Plugins/Editor/AssetManagerEditor/Source/AssetManagerEditor/Private/Insights/Table/ViewModels/BaseTreeNode.h
--- a/Engine/Plugins/Editor/AssetManagerEditor/Source/AssetManagerEditor/Private/Insights/Table/ViewModels/BaseTreeNode.h
+++ b/Engine/Plugins/Editor/AssetManagerEditor/Source/AssetManagerEditor/Private/Insights/Table/ViewModels/BaseTreeNode.h
@@ -59,8 +59,10 @@ public:
 	FBaseTreeNode(const FName InName, bool bInIsGroup)
 		: DefaultSortOrder(0)
 		, Name(InName)
+		, IconBrush(nullptr)
 		, GroupData(bInIsGroup ? new FGroupNodeData() : &DefaultGroupData)
 	{
+		SetDefaultIcon();
 	}
 
 	virtual ~FBaseTreeNode()
@@ -120,6 +122,32 @@ public:
 		}
 	}
 
+	/**
+	 * @return a brush icon for this node.
+	 */
+	const FSlateBrush* GetIcon() const
+	{
+		return IconBrush;
+	}
+
+	/**
+	 * Sets an icon brush for this node.
+	 */
+	void SetIcon(const FSlateBrush* InIconBrush)
+	{
+		IconBrush = InIconBrush;
+	}
+
+	/**
+	 * Sets the default icon brush for this node.
+	 */
+	void SetDefaultIcon();
+
+	/**
+	 * @return the color tint for icon and name text.
+	 */
+	virtual FLinearColor GetColor() const;
+
 	/**
 	 * @return a pointer to a data context for this node.
 	 */
@@ -292,6 +320,9 @@ private:
 	/** The name of this node. */
 	const FName Name;
 
+	/** The icon of this node. */
+	const FSlateBrush* IconBrush;
+
 	/** A weak pointer to the group/parent of this node. */
 	FBaseTreeNodeWeak GroupPtr;
 
diff --git a/Engine/Plugins/Editor/AssetManagerEditor/Source/AssetManagerEditor/Private/Insights/Table/Widgets/STableTreeViewCell.cpp b/Engine/Plugins/Editor/AssetManagerEditor/Source/AssetManagerEditor/Private/Insights/Table/Widgets/STableTreeViewCell.cpp
--- a/Engine/Plugins/Editor/AssetManagerEditor/Source/AssetManagerEditor/Private/Insights/Table/Widgets/STableTreeViewCell.cpp
+++ b/Engine/Plugins/Editor/AssetManagerEditor/Source/AssetManagerEditor/Private/Insights/Table/Widgets/STableTreeViewCell.cpp
@@ -75,24 +75,34 @@ TSharedRef<SWidget> STableTreeViewCell::GenerateWidgetForNameColumn(const FArgum
 			SNew(SExpanderArrow, TableRow)
 		]
 
-		// Info icon + tooltip
+		// Icon + tooltip
 		+ SHorizontalBox::Slot()
 		.AutoWidth()
 		.HAlign(HAlign_Center)
 		.VAlign(VAlign_Center)
 		[
-			SNew(SImage)
-			.Visibility(this, &STableTreeViewCell::GetHintIconVisibility)
-			.Image(FInsightsStyle::GetBrush("Icons.Hint.TreeItem"))
-			.ToolTip(GetRowToolTip(TableRow))
+			SNew(SOverlay)
+			+ SOverlay::Slot()
+			[
+				SNew(SImage)
+				.Image(this, &STableTreeViewCell::GetIcon)
+				.ColorAndOpacity(this, &STableTreeViewCell::GetIconColorAndOpacity)
+			]
+			+ SOverlay::Slot()
+			[
+				SNew(SImage)
+				.Visibility(this, &STableTreeViewCell::GetHintIconVisibility)
+				.Image(FInsightsStyle::GetBrush("Icons.Hint.TreeItem"))
+				.ToolTip(GetRowToolTip(TableRow))
+			]
 		]
 
 		// Name
 		+ SHorizontalBox::Slot()
 		.AutoWidth()
-		.VAlign(VAlign_Center)
+		.VAlign(VAlign_Top)
 		.HAlign(HAlign_Left)
-		.Padding(FMargin(2.0f, 0.0f))
+		.Padding(FMargin(2.0f, 1.5f, 2.0f, 0.0f))
 		[
 			SNew(STextBlock)
 			.Text(this, &STableTreeViewCell::GetDisplayName)
@@ -105,9 +115,9 @@ TSharedRef<SWidget> STableTreeViewCell::GenerateWidgetForNameColumn(const FArgum
 		// Name Suffix
 		+ SHorizontalBox::Slot()
 		.AutoWidth()
-		.VAlign(VAlign_Center)
+		.VAlign(VAlign_Top)
 		.HAlign(HAlign_Left)
-		.Padding(FMargin(2.0f, 0.0f))
+		.Padding(FMargin(2.0f, 1.5f, 2.0f, 0.0f))
 		[
 			SNew(STextBlock)
 			.Visibility(this, &STableTreeViewCell::HasExtraDisplayName)
@@ -159,9 +169,9 @@ TSharedRef<SWidget> STableTreeViewCell::GenerateWidgetForTableColumn(const FArgu
 
 		+ SHorizontalBox::Slot()
 		.FillWidth(1.0f)
-		.VAlign(VAlign_Center)
+		.VAlign(VAlign_Top)
 		.HAlign(ColumnPtr->GetHorizontalAlignment())
-		.Padding(FMargin(2.0f, 0.0f))
+		.Padding(FMargin(2.0f, 1.5f, 2.0f, 0.0f))
 		[
 			TextBox
 		];
diff --git a/Engine/Plugins/Editor/AssetManagerEditor/Source/AssetManagerEditor/Private/Insights/Table/Widgets/STableTreeViewCell.h b/Engine/Plugins/Editor/AssetManagerEditor/Source/AssetManagerEditor/Private/Insights/Table/Widgets/STableTreeViewCell.h
--- a/Engine/Plugins/Editor/AssetManagerEditor/Source/AssetManagerEditor/Private/Insights/Table/Widgets/STableTreeViewCell.h
+++ b/Engine/Plugins/Editor/AssetManagerEditor/Source/AssetManagerEditor/Private/Insights/Table/Widgets/STableTreeViewCell.h
@@ -105,6 +105,11 @@ protected:
 
 	TSharedPtr<class IToolTip> GetRowToolTip(const TSharedRef<ITableRow>& TableRow) const;
 
+	const FSlateBrush* GetIcon() const
+	{
+		return TableTreeNodePtr->GetIcon();
+	}
+
 	FText GetDisplayName() const
 	{
 		return TableTreeNodePtr->GetDisplayName();
@@ -122,12 +127,21 @@ protected:
 
 	FText GetValueAsText() const;
 
+	FSlateColor GetIconColorAndOpacity() const
+	{
+		const FLinearColor TextColor =
+			IsHovered() ?
+				TableTreeNodePtr->GetColor() :
+				TableTreeNodePtr->GetColor().CopyWithNewOpacity(0.5f);
+		return TextColor;
+	}
+
 	FSlateColor GetColorAndOpacity() const
 	{
 		const FLinearColor TextColor =
 			TableTreeNodePtr->IsFiltered() ?
-				FLinearColor(1.0f, 1.0f, 1.0f, 0.5f) :
-				FLinearColor(1.0f, 1.0f, 1.0f, 1.0f);
+				TableTreeNodePtr->GetColor().CopyWithNewOpacity(0.5f) :
+				TableTreeNodePtr->GetColor();
 		return TextColor;
 	}
 
@@ -155,7 +169,7 @@ protected:
 	{
 		const FLinearColor ShadowColor =
 			TableTreeNodePtr->IsFiltered() ?
-				FLinearColor(0.f, 0.f, 0.f, 0.25f) :
+				FLinearColor(0.0f, 0.0f, 0.0f, 0.25f) :
 				FLinearColor(0.0f, 0.0f, 0.0f, 0.5f);
 		return ShadowColor;
 	}
diff --git a/Engine/Plugins/Editor/AssetManagerEditor/Source/AssetManagerEditor/Private/SAssetAuditBrowser.cpp b/Engine/Plugins/Editor/AssetManagerEditor/Source/AssetManagerEditor/Private/SAssetAuditBrowser.cpp
--- a/Engine/Plugins/Editor/AssetManagerEditor/Source/AssetManagerEditor/Private/SAssetAuditBrowser.cpp
+++ b/Engine/Plugins/Editor/AssetManagerEditor/Source/AssetManagerEditor/Private/SAssetAuditBrowser.cpp
@@ -31,6 +31,8 @@
 
 #define LOCTEXT_NAMESPACE "AssetManagementBrowser"
 
+//#pragma optimize("", off)
+
 const FString SAssetAuditBrowser::SettingsIniSection = TEXT("AssetManagementBrowser");
 const int32 SAssetAuditBrowser::MaxAssetsHistory = 10;
 
@@ -723,6 +725,26 @@ void SAssetAuditBrowser::GoToHistoryIndex(int32 InHistoryIdx)
 	}
 }
 
+#if ASSET_TABLE_TREE_VIEW_ENABLED
+void SAssetAuditBrowser::PopulateAssetTableRow(const FAssetData& AssetData, class FAssetTableRow* OutRow)
+{
+	OutRow->Type = AssetData.AssetClassPath.ToString();
+	OutRow->Name = AssetData.AssetName.ToString();
+	OutRow->Path = AssetData.GetObjectPathString();
+	EditorModule->GetStringValueForCustomColumn(AssetData, FPrimaryAssetId::PrimaryAssetTypeTag, OutRow->PrimaryType);
+	EditorModule->GetStringValueForCustomColumn(AssetData, FPrimaryAssetId::PrimaryAssetNameTag, OutRow->PrimaryName);
+	//EditorModule->GetIntegerValueForCustomColumn(AssetData, IAssetManagerEditorModule::ManagedDiskSizeName, OutRow->ManagedDiskSize);
+	//EditorModule->GetIntegerValueForCustomColumn(AssetData, IAssetManagerEditorModule::DiskSizeName, OutRow->DiskSize);
+	EditorModule->GetIntegerValueForCustomColumn(AssetData, IAssetManagerEditorModule::StageChunkCompressedSizeName, OutRow->StagedCompressedSize);
+	EditorModule->GetIntegerValueForCustomColumn(AssetData, IAssetManagerEditorModule::TotalUsageName, OutRow->TotalUsageCount);
+	//EditorModule->GetStringValueForCustomColumn(AssetData, IAssetManagerEditorModule::CookRuleName, OutRow->CookRule);
+	//EditorModule->GetStringValueForCustomColumn(AssetData, IAssetManagerEditorModule::ChunksName, OutRow->Chunks);
+	EditorModule->GetStringValueForCustomColumn(AssetData, IAssetManagerEditorModule::GameFeaturePluginsName, OutRow->GameFeaturePlugin);
+	OutRow->NativeClass = AssetData.AssetClassPath.GetAssetName().ToString();
+	OutRow->Color = USlateThemeManager::Get().GetColor((EStyleColor)((int32)EStyleColor::AccentBlue + FMath::Rand() % 8)); // TODO: color by type!?
+}
+#endif
+
 void SAssetAuditBrowser::RefreshAssetView()
 {
 	FARFilter Filter;
@@ -750,26 +772,92 @@ void SAssetAuditBrowser::RefreshAssetView()
 		TSharedPtr<FAssetTable> AssetTable = AssetTableTreeView->GetAssetTable();
 		if (AssetTable.IsValid())
 		{
-			TArray<FAssetTableRow>& Assets = AssetTable->GetAssets();
-			Assets.Reset();
+			AssetTable->ClearAllData();
+
+			TMap<FAssetData, int32> AssetToIndexMap;
 
 			TArray<FAssetData> SourceAssets;
 			AssetRegistry->GetAssets(Filter, SourceAssets);
-			for (const FAssetData& SourceAsset : SourceAssets)
+			for (int32 Index = 0; Index < SourceAssets.Num(); Index++)
 			{
-				FAssetTableRow& Asset = Assets.AddDefaulted_GetRef();
-				Asset.Type = SourceAsset.AssetClassPath.ToString();
-				Asset.Name = SourceAsset.AssetName.ToString();
-				Asset.Path = SourceAsset.GetObjectPathString();
-				EditorModule->GetStringValueForCustomColumn(SourceAsset, FPrimaryAssetId::PrimaryAssetTypeTag, Asset.PrimaryType);
-				EditorModule->GetStringValueForCustomColumn(SourceAsset, FPrimaryAssetId::PrimaryAssetNameTag, Asset.PrimaryName);
-				EditorModule->GetIntegerValueForCustomColumn(SourceAsset, IAssetManagerEditorModule::ManagedDiskSizeName, Asset.ManagedDiskSize);
-				EditorModule->GetIntegerValueForCustomColumn(SourceAsset, IAssetManagerEditorModule::DiskSizeName, Asset.DiskSize);
-				EditorModule->GetIntegerValueForCustomColumn(SourceAsset, IAssetManagerEditorModule::StageChunkCompressedSizeName, Asset.StagedCompressedSize);
-				EditorModule->GetIntegerValueForCustomColumn(SourceAsset, IAssetManagerEditorModule::TotalUsageName, Asset.TotalUsageCount);
-				EditorModule->GetStringValueForCustomColumn(SourceAsset, IAssetManagerEditorModule::CookRuleName, Asset.CookRule);
-				EditorModule->GetStringValueForCustomColumn(SourceAsset, IAssetManagerEditorModule::ChunksName, Asset.Chunks);
-				Asset.NativeClass = SourceAsset.AssetClassPath.GetAssetName().ToString();
+				const FAssetData& SourceAsset = SourceAssets[Index];
+				FAssetTableRow AssetRow;
+				PopulateAssetTableRow(SourceAsset, &AssetRow);
+				AssetTable->AddAsset(AssetRow);
+				AssetToIndexMap.Add(SourceAsset, Index);
+			}
+			AssetTable->SetVisibleAssetCount(AssetTable->GetAssets().Num());
+
+			int32 EntryIndex = AssetTable->GetAssets().Num();
+			const FAssetManagerEditorRegistrySource* RegistrySource = EditorModule->GetCurrentRegistrySource();
+
+			TSet<FName> AlreadyProcessedPackages;
+			// Go over all SourceAssets. We will add more to the end as we find we need to add them.
+			for (int32 SourceAssetIndex = 0; SourceAssetIndex < SourceAssets.Num(); SourceAssetIndex++)
+			{
+				// There are potentially multiple assets in a package, each with an FAssetData
+				// We are only dealing in package dependencies, though, so don't process the same package twice
+				// But we will need to conceptually assign all the package dependencies to all assets in the package below
+				if (AlreadyProcessedPackages.Contains(SourceAssets[SourceAssetIndex].PackageName))
+				{
+					continue;
+				}
+				else
+				{
+					AlreadyProcessedPackages.Add(SourceAssets[SourceAssetIndex].PackageName);
+				}
+
+
+				FAssetIdentifier CurrentIdentifier(SourceAssets[SourceAssetIndex].PackageName);
+				// We'll have to add the dependencies to all these entries
+				TArray<FAssetData> AssetsInSourcePackage;
+				AssetRegistry->GetAssetsByPackageName(SourceAssets[SourceAssetIndex].PackageName, AssetsInSourcePackage, /*bIncludeOnlyOnDiskAssets*/ true); // Only use on disk assets to avoid creating FAssetData for everything in memory
+
+				// Get the dependencies
+				TArray<FAssetIdentifier> DependencyList;
+//				RegistrySource->GetDependencies(CurrentIdentifier, DependencyList, UE::AssetRegistry::EDependencyCategory::Package);
+				AssetRegistry->GetDependencies(CurrentIdentifier, DependencyList, UE::AssetRegistry::EDependencyCategory::Package);
+
+				TArray<int32> IndicesOfDependenciesInRowTableToAddToCurrentSourceAssetRow;
+
+				for (const FAssetIdentifier& Dependency : DependencyList)
+				{
+					TArray<FAssetData> AssetsInDependencyPackage;
+					if (AssetRegistry->GetAssetsByPackageName(Dependency.PackageName, AssetsInDependencyPackage, /*bIncludeOnlyOnDiskAssets*/ true))
+					{
+						for (const FAssetData& DependencyAsset : AssetsInDependencyPackage)
+						{
+							if (int32* DependencyIndex = AssetToIndexMap.Find(DependencyAsset))
+							{
+								// Do we already know about the asset? If so, great.
+								IndicesOfDependenciesInRowTableToAddToCurrentSourceAssetRow.Add(*DependencyIndex);
+							}
+							else
+							{
+								// We don't know about this asset yet. Create a new row for it and add it to the source list for further dependency analysis
+								FAssetTableRow NewRow;
+								PopulateAssetTableRow(DependencyAsset, &NewRow);
+								AssetTable->AddAsset(NewRow);
+								IndicesOfDependenciesInRowTableToAddToCurrentSourceAssetRow.Add(SourceAssets.Num());
+								AssetToIndexMap.Add(DependencyAsset, SourceAssets.Num());
+								SourceAssets.Add(DependencyAsset);
+							}
+						}
+					}
+				}
+
+				if (IndicesOfDependenciesInRowTableToAddToCurrentSourceAssetRow.Num())
+				{
+					// We found some dependencies. Let's add them.
+					for (const FAssetData& SourceAssetData : AssetsInSourcePackage)
+					{
+						int32* RowIndex = AssetToIndexMap.Find(SourceAssetData);
+						if (ensure(RowIndex != nullptr))
+						{
+							AssetTable->GetAsset(*RowIndex)->Dependencies.Append(IndicesOfDependenciesInRowTableToAddToCurrentSourceAssetRow);
+						}
+					}
+				}
 			}
 
 			AssetTableTreeView->RebuildTreeAsync();
diff --git a/Engine/Plugins/Editor/AssetManagerEditor/Source/AssetManagerEditor/Private/SAssetAuditBrowser.h b/Engine/Plugins/Editor/AssetManagerEditor/Source/AssetManagerEditor/Private/SAssetAuditBrowser.h
--- a/Engine/Plugins/Editor/AssetManagerEditor/Source/AssetManagerEditor/Private/SAssetAuditBrowser.h
+++ b/Engine/Plugins/Editor/AssetManagerEditor/Source/AssetManagerEditor/Private/SAssetAuditBrowser.h
@@ -132,6 +132,9 @@ protected:
 	/** Refresh the asset view with a new filter */
 	void RefreshAssetView();
 
+#if ASSET_TABLE_TREE_VIEW_ENABLED
+	void PopulateAssetTableRow(const FAssetData& AssetData, class FAssetTableRow* OutRow);
+#endif
 	/** Commands handled by this widget */
 	TSharedPtr<FUICommandList> Commands;
 
diff --git a/Engine/Plugins/Editor/AssetManagerEditor/Source/AssetManagerEditor/Private/TreeView/AssetDependencyGrouping.cpp b/Engine/Plugins/Editor/AssetManagerEditor/Source/AssetManagerEditor/Private/TreeView/AssetDependencyGrouping.cpp
new file mode 100644
--- /dev/null
+++ b/Engine/Plugins/Editor/AssetManagerEditor/Source/AssetManagerEditor/Private/TreeView/AssetDependencyGrouping.cpp
@@ -0,0 +1,92 @@
+// Copyright Epic Games, Inc. All Rights Reserved.
+
+#include "AssetDependencyGrouping.h"
+#include "AssetTreeNode.h"
+#include "Insights/Common/AsyncOperationProgress.h"
+
+#define LOCTEXT_NAMESPACE "FAssetDependencyGrouping"
+
+INSIGHTS_IMPLEMENT_RTTI(FAssetDependencyGrouping)
+
+FAssetDependencyGrouping::FAssetDependencyGrouping()
+	: FTreeNodeGrouping(
+		LOCTEXT("Grouping_ByDependency_ShortName", "Dependency"),
+		LOCTEXT("Grouping_ByDependency_TitleName", "By Dependency"),
+		LOCTEXT("Grouping_ByDependency_Desc", "Group assets based on their dependency."),
+		TEXT("Icons.Group.TreeItem"),
+		nullptr)
+{
+}
+
+FAssetDependencyGrouping::~FAssetDependencyGrouping()
+{
+}
+
+void FAssetDependencyGrouping::GroupNodes(const TArray<UE::Insights::FTableTreeNodePtr>& Nodes, UE::Insights::FTableTreeNode& ParentGroup, TWeakPtr<UE::Insights::FTable> InParentTable, UE::Insights::IAsyncOperationProgress& InAsyncOperationProgress) const
+{
+	using namespace UE::Insights;
+
+	ParentGroup.ClearChildren();
+
+	TSharedPtr<FAssetTable> AssetTable = StaticCastSharedPtr<FAssetTable>(InParentTable.Pin());
+	check(AssetTable.IsValid());
+
+	for (FTableTreeNodePtr NodePtr : Nodes)
+	{
+		if (InAsyncOperationProgress.ShouldCancelAsyncOp())
+		{
+			return;
+		}
+
+		if (NodePtr->IsGroup() || !NodePtr->Is<FAssetTreeNode>())
+		{
+			ParentGroup.AddChildAndSetGroupPtr(NodePtr);
+			continue;
+		}
+
+		// For each (visible) asset, we will create the following hierarcy:
+		//
+		// |
+		// +-- [group:{AssetName}]
+		// |   |
+		// |   +-- [asset:{AssetName}] // current asset
+		// |   |
+		// |   +-- [group:Dependencies]
+		// |       |
+		// |       +-- [asset:{DependentAsset1}] // dependent of current asset
+		// |       |
+		// |       +-- [asset:{DependentAsset2}] // another dependent of current asset
+		// |       |
+
+		FAssetTreeNode& AssetNode = NodePtr->As<FAssetTreeNode>();
+		const FAssetTableRow& Asset = AssetNode.GetAssetChecked();
+
+		FTableTreeNodePtr GroupPtr = MakeShared<FTableTreeNode>(FName(Asset.GetName(), 0), InParentTable);
+		GroupPtr->SetExpansion(false);
+		ParentGroup.AddChildAndSetGroupPtr(GroupPtr);
+
+		GroupPtr->AddChildAndSetGroupPtr(NodePtr);
+
+		FTableTreeNodePtr DepGroupPtr = nullptr;
+		const TArray<int32>& Dependencies = Asset.GetDependencies();
+		if (Dependencies.Num() > 0)
+		{
+			DepGroupPtr = MakeShared<FTableTreeNode>(FName(TEXT("Dependencies"), 0), InParentTable);
+			DepGroupPtr->SetExpansion(false);
+			GroupPtr->AddChildAndSetGroupPtr(DepGroupPtr);
+		}
+		for (int32 DepAssetIndex : Dependencies)
+		{
+			if (!ensure(AssetTable->IsValidRowIndex(DepAssetIndex)))
+			{
+				continue;
+			}
+			const FAssetTableRow& DepAsset = AssetTable->GetAssetChecked(DepAssetIndex);
+			FName DepAssetNodeName(DepAsset.GetName());
+			FAssetTreeNodePtr DepAssetNodePtr = MakeShared<FAssetTreeNode>(DepAssetNodeName, AssetTable, DepAssetIndex);
+			DepGroupPtr->AddChildAndSetGroupPtr(DepAssetNodePtr);
+		}
+	}
+}
+
+#undef LOCTEXT_NAMESPACE
diff --git a/Engine/Plugins/Editor/AssetManagerEditor/Source/AssetManagerEditor/Private/TreeView/AssetDependencyGrouping.h b/Engine/Plugins/Editor/AssetManagerEditor/Source/AssetManagerEditor/Private/TreeView/AssetDependencyGrouping.h
new file mode 100644
--- /dev/null
+++ b/Engine/Plugins/Editor/AssetManagerEditor/Source/AssetManagerEditor/Private/TreeView/AssetDependencyGrouping.h
@@ -0,0 +1,21 @@
+// Copyright Epic Games, Inc. All Rights Reserved.
+
+#pragma once
+
+#include "CoreMinimal.h"
+#include "Insights/Table/ViewModels/BaseTreeNode.h"
+#include "Insights/Table/ViewModels/TreeNodeGrouping.h"
+
+class FAssetDependencyGrouping : public UE::Insights::FTreeNodeGrouping
+{
+	INSIGHTS_DECLARE_RTTI(FAssetDependencyGrouping, UE::Insights::FTreeNodeGrouping)
+
+public:
+	FAssetDependencyGrouping();
+	virtual ~FAssetDependencyGrouping();
+
+	virtual void GroupNodes(const TArray<UE::Insights::FTableTreeNodePtr>& Nodes, UE::Insights::FTableTreeNode& ParentGroup, TWeakPtr<UE::Insights::FTable> InParentTable, UE::Insights::IAsyncOperationProgress& InAsyncOperationProgress) const override;
+
+private:
+	//...
+};
diff --git a/Engine/Plugins/Editor/AssetManagerEditor/Source/AssetManagerEditor/Private/TreeView/AssetTable.cpp b/Engine/Plugins/Editor/AssetManagerEditor/Source/AssetManagerEditor/Private/TreeView/AssetTable.cpp
--- a/Engine/Plugins/Editor/AssetManagerEditor/Source/AssetManagerEditor/Private/TreeView/AssetTable.cpp
+++ b/Engine/Plugins/Editor/AssetManagerEditor/Source/AssetManagerEditor/Private/TreeView/AssetTable.cpp
@@ -25,13 +25,14 @@ const FName FAssetTableColumns::TypeColumnId(TEXT("Type"));
 const FName FAssetTableColumns::PathColumnId(TEXT("Path"));
 const FName FAssetTableColumns::PrimaryTypeColumnId(TEXT("PrimaryType"));
 const FName FAssetTableColumns::PrimaryNameColumnId(TEXT("PrimaryName"));
-const FName FAssetTableColumns::ManagedDiskSizeColumnId(TEXT("ManagedDiskSize"));
-const FName FAssetTableColumns::DiskSizeColumnId(TEXT("DiskSize"));
+//const FName FAssetTableColumns::ManagedDiskSizeColumnId(TEXT("ManagedDiskSize"));
+//const FName FAssetTableColumns::DiskSizeColumnId(TEXT("DiskSize"));
 const FName FAssetTableColumns::StagedCompressedSizeColumnId(TEXT("StagedCompressedSize"));
 const FName FAssetTableColumns::TotalUsageCountColumnId(TEXT("TotalUsageCount"));
-const FName FAssetTableColumns::CookRuleColumnId(TEXT("CookRule"));
-const FName FAssetTableColumns::ChunksColumnId(TEXT("Chunks"));
+//const FName FAssetTableColumns::CookRuleColumnId(TEXT("CookRule"));
+//const FName FAssetTableColumns::ChunksColumnId(TEXT("Chunks"));
 const FName FAssetTableColumns::NativeClassColumnId(TEXT("NativeClass"));
+const FName FAssetTableColumns::GameFeaturePluginColumnId(TEXT("GameFeaturePlugin"));
 
 ////////////////////////////////////////////////////////////////////////////////////////////////////
 // FAssetTable
@@ -40,12 +41,17 @@ const FName FAssetTableColumns::NativeClassColumnId(TEXT("NativeClass"));
 FAssetTable::FAssetTable()
 {
 #if 0 // debug, mock data
-	constexpr int AssetCount = 100;
-	Assets.Reserve(AssetCount);
-	for (int AssetIndex = 0; AssetIndex < AssetCount; ++AssetIndex)
+
+	VisibleAssetCount = 100;
+	constexpr int32 HiddenAssetCount = 50;
+	const int32 TotalAssetCount = VisibleAssetCount + HiddenAssetCount;
+	
+	// Create assets.
+	Assets.Reserve(TotalAssetCount);
+	for (int32 AssetIndex = 0; AssetIndex < TotalAssetCount; ++AssetIndex)
 	{
 		FAssetTableRow& Asset = Assets.AddDefaulted_GetRef();
-		
+
 		int32 Id = FMath::Rand();
 		int32 Id2 = FMath::Rand();
 
@@ -54,14 +60,45 @@ FAssetTable::FAssetTable()
 		Asset.Path = FString::Printf(TEXT("A%02d/B%02d/C%02d/D%02d"), Id % 11, Id % 13, Id % 17, Id % 19);
 		Asset.PrimaryType = FString::Printf(TEXT("PT_%02d"), Id2 % 10);
 		Asset.PrimaryName = FString::Printf(TEXT("PN%d"), Id2);
-		Asset.ManagedDiskSize = FMath::Abs(Id * Id);
-		Asset.DiskSize = FMath::Abs(Id * Id * Id);
+		//Asset.ManagedDiskSize = FMath::Abs(Id * Id);
+		//Asset.DiskSize = FMath::Abs(Id * Id * Id);
 		Asset.StagedCompressedSize = Asset.DiskSize / 2;
 		Asset.TotalUsageCount = Id % 1000;
-		Asset.CookRule = FString::Printf(TEXT("CookRule%02d"), (Id * Id) % 8);
-		Asset.Chunks = FString::Printf(TEXT("Chunks%02d"), (Id * Id + 1) % 41);
+		//Asset.CookRule = FString::Printf(TEXT("CookRule%02d"), (Id * Id) % 8);
+		//Asset.Chunks = FString::Printf(TEXT("Chunks%02d"), (Id * Id + 1) % 41);
 		Asset.NativeClass = FString::Printf(TEXT("NativeClass%02d"), (Id * Id * Id) % 8);
+		Asset.GameFeaturePlugin = FString::Printf(TEXT("GFP_%02d"), (Id * Id) % 7);
+	}
+
+	// Set dependencies (only for visible assets)
+	for (int32 AssetIndex = 0; AssetIndex < VisibleAssetCount; ++AssetIndex)
+	{
+		if (FMath::Rand() % 100 > 60) // 60% chance to have dependencies
+		{
+			continue;
+		}
+
+		FAssetTableRow& Asset = Assets[AssetIndex];
+
+		int32 NumDependents = FMath::Rand() % 30; // max 30 dependent assets
+		while (--NumDependents >= 0)
+		{
+			int32 DepIndex = -1;
+			if (FMath::Rand() % 100 <= 10) // 10% chance to be another asset that is visible by default
+			{
+				DepIndex = FMath::Rand() % VisibleAssetCount;
+			}
+			else // 90% chance to be a dependet only asset (not visible by default)
+			{
+				DepIndex = VisibleAssetCount + FMath::Rand() % HiddenAssetCount;
+			}
+			if (!Asset.Dependencies.Contains(DepIndex))
+			{
+				Asset.Dependencies.Add(DepIndex);
+			}
+		}
 	}
+
 #endif
 }
 
@@ -437,6 +474,7 @@ void FAssetTable::AddDefaultColumns()
 	}
 	//////////////////////////////////////////////////
 	// Managed Disk Size Column
+#if 0
 	{
 		TSharedRef<FTableColumn> ColumnRef = MakeShared<FTableColumn>(FAssetTableColumns::ManagedDiskSizeColumnId);
 		FTableColumn& Column = *ColumnRef;
@@ -494,8 +532,10 @@ void FAssetTable::AddDefaultColumns()
 
 		AddColumn(ColumnRef);
 	}
+#endif
 	//////////////////////////////////////////////////
 	// Exclusive Disk Size Column
+#if 0
 	{
 		TSharedRef<FTableColumn> ColumnRef = MakeShared<FTableColumn>(FAssetTableColumns::DiskSizeColumnId);
 		FTableColumn& Column = *ColumnRef;
@@ -553,6 +593,7 @@ void FAssetTable::AddDefaultColumns()
 
 		AddColumn(ColumnRef);
 	}
+#endif
 	//////////////////////////////////////////////////
 	// Staged Compressed Size Column
 	{
@@ -670,6 +711,7 @@ void FAssetTable::AddDefaultColumns()
 	}
 	//////////////////////////////////////////////////
 	// Cook Rule Column
+#if 0
 	{
 		TSharedRef<FTableColumn> ColumnRef = MakeShared<FTableColumn>(FAssetTableColumns::CookRuleColumnId);
 		FTableColumn& Column = *ColumnRef;
@@ -723,8 +765,10 @@ void FAssetTable::AddDefaultColumns()
 
 		AddColumn(ColumnRef);
 	}
+#endif
 	//////////////////////////////////////////////////
 	// Chunks Column
+#if 0
 	{
 		TSharedRef<FTableColumn> ColumnRef = MakeShared<FTableColumn>(FAssetTableColumns::ChunksColumnId);
 		FTableColumn& Column = *ColumnRef;
@@ -778,6 +822,7 @@ void FAssetTable::AddDefaultColumns()
 
 		AddColumn(ColumnRef);
 	}
+#endif
 	//////////////////////////////////////////////////
 	// Native Class Column
 	{
@@ -834,6 +879,61 @@ void FAssetTable::AddDefaultColumns()
 		AddColumn(ColumnRef);
 	}
 	//////////////////////////////////////////////////
+	// GameFeaturePlugin Column
+	{
+		TSharedRef<FTableColumn> ColumnRef = MakeShared<FTableColumn>(FAssetTableColumns::GameFeaturePluginColumnId);
+		FTableColumn& Column = *ColumnRef;
+
+		Column.SetIndex(ColumnIndex++);
+
+		Column.SetShortName(LOCTEXT("GameFeaturePluginColumnName", "GameFeaturePlugin"));
+		Column.SetTitleName(LOCTEXT("GameFeaturePluginColumnTitle", "GameFeaturePlugin"));
+		Column.SetDescription(LOCTEXT("GameFeaturePluginColumnDesc", "GameFeaturePlugin of the asset"));
+
+		Column.SetFlags(ETableColumnFlags::ShouldBeVisible | ETableColumnFlags::CanBeHidden | ETableColumnFlags::CanBeFiltered);
+
+		Column.SetHorizontalAlignment(HAlign_Left);
+		Column.SetInitialWidth(120.0f);
+
+		Column.SetDataType(ETableCellDataType::CString);
+
+		class FNativeClassValueGetter : public FTableCellValueGetter
+		{
+		public:
+			virtual const TOptional<FTableCellValue> GetValue(const FTableColumn& Column, const FBaseTreeNode& Node) const
+			{
+				if (Node.IsGroup())
+				{
+					const FTableTreeNode& NodePtr = static_cast<const FTableTreeNode&>(Node);
+					if (NodePtr.HasAggregatedValue(Column.GetId()))
+					{
+						return NodePtr.GetAggregatedValue(Column.GetId());
+					}
+				}
+				else //if (Node->Is<FAssetTreeNode>())
+				{
+					const FAssetTreeNode& TreeNode = static_cast<const FAssetTreeNode&>(Node);
+					const FAssetTableRow& Asset = TreeNode.GetAssetChecked();
+					return FTableCellValue(Asset.GetGameFeaturePlugin());
+				}
+
+				return TOptional<FTableCellValue>();
+			}
+		};
+		TSharedRef<ITableCellValueGetter> Getter = MakeShared<FNativeClassValueGetter>();
+		Column.SetValueGetter(Getter);
+
+		TSharedRef<ITableCellValueFormatter> Formatter = MakeShared<FCStringValueFormatterAsText>();
+		Column.SetValueFormatter(Formatter);
+
+		TSharedRef<ITableCellValueSorter> Sorter = MakeShared<FSorterByCStringValue>(ColumnRef);
+		Column.SetValueSorter(Sorter);
+
+		Column.SetAggregation(ETableColumnAggregation::SameValue);
+
+		AddColumn(ColumnRef);
+	}
+	//////////////////////////////////////////////////
 }
 
 ////////////////////////////////////////////////////////////////////////////////////////////////////
diff --git a/Engine/Plugins/Editor/AssetManagerEditor/Source/AssetManagerEditor/Private/TreeView/AssetTable.h b/Engine/Plugins/Editor/AssetManagerEditor/Source/AssetManagerEditor/Private/TreeView/AssetTable.h
--- a/Engine/Plugins/Editor/AssetManagerEditor/Source/AssetManagerEditor/Private/TreeView/AssetTable.h
+++ b/Engine/Plugins/Editor/AssetManagerEditor/Source/AssetManagerEditor/Private/TreeView/AssetTable.h
@@ -22,13 +22,14 @@ struct FAssetTableColumns
 	static const FName PathColumnId;
 	static const FName PrimaryTypeColumnId;
 	static const FName PrimaryNameColumnId;
-	static const FName ManagedDiskSizeColumnId;
-	static const FName DiskSizeColumnId;
+	//static const FName ManagedDiskSizeColumnId;
+	//static const FName DiskSizeColumnId;
 	static const FName StagedCompressedSizeColumnId;
 	static const FName TotalUsageCountColumnId;
-	static const FName CookRuleColumnId;
-	static const FName ChunksColumnId;
+	//static const FName CookRuleColumnId;
+	//static const FName ChunksColumnId;
 	static const FName NativeClassColumnId;
+	static const FName GameFeaturePluginColumnId;
 };
 
 class FAssetTableRow
@@ -49,13 +50,16 @@ public:
 	const TCHAR* GetPath() const { return *Path; }
 	const TCHAR* GetPrimaryType() const { return *PrimaryType; }
 	const TCHAR* GetPrimaryName() const { return *PrimaryName; }
-	int64 GetManagedDiskSize() const { return ManagedDiskSize; }
-	int64 GetDiskSize() const { return DiskSize; }
+	//int64 GetManagedDiskSize() const { return ManagedDiskSize; }
+	//int64 GetDiskSize() const { return DiskSize; }
 	int64 GetStagedCompressedSize() const { return StagedCompressedSize; }
 	int64 GetTotalUsageCount() const { return TotalUsageCount; }
-	const TCHAR* GetCookRule() const { return *CookRule; }
-	const TCHAR* GetChunks() const { return *Chunks; }
+	//const TCHAR* GetCookRule() const { return *CookRule; }
+	//const TCHAR* GetChunks() const { return *Chunks; }
 	const TCHAR* GetNativeClass() const { return *NativeClass; }
+	const TArray<int32>& GetDependencies() const { return Dependencies; }
+	const TCHAR* GetGameFeaturePlugin() const { return *GameFeaturePlugin; }
+	FLinearColor GetColor() const { return Color; }
 
 private:
 	FString Type;
@@ -63,13 +67,16 @@ private:
 	FString Path;
 	FString PrimaryType;
 	FString PrimaryName;
-	int64 ManagedDiskSize = 0;
-	int64 DiskSize = 0;
+	//int64 ManagedDiskSize = 0;
+	//int64 DiskSize = 0;
 	int64 StagedCompressedSize = 0;
 	int64 TotalUsageCount = 0;
-	FString CookRule;
-	FString Chunks;
+	//FString CookRule;
+	//FString Chunks;
 	FString NativeClass;
+	TArray<int32> Dependencies;
+	FString GameFeaturePlugin;
+	FLinearColor Color;
 };
 
 class FAssetTable : public UE::Insights::FTable
@@ -84,12 +91,29 @@ public:
 	const TArray<FAssetTableRow>& GetAssets() const { return Assets; }
 
 	bool IsValidRowIndex(int32 InIndex) const { return InIndex >= 0 && InIndex < Assets.Num(); }
+	FAssetTableRow* GetAsset(int32 InIndex) { return IsValidRowIndex(InIndex) ? &Assets[InIndex] : nullptr; }
 	const FAssetTableRow* GetAsset(int32 InIndex) const { return IsValidRowIndex(InIndex) ? &Assets[InIndex] : nullptr; }
+	FAssetTableRow& GetAssetChecked(int32 InIndex) { check(IsValidRowIndex(InIndex)); return Assets[InIndex]; }
 	const FAssetTableRow& GetAssetChecked(int32 InIndex) const { check(IsValidRowIndex(InIndex)); return Assets[InIndex]; }
 
+	int32 GetTotalAssetCount() const { return Assets.Num(); }
+	int32 GetVisibleAssetCount() const { return VisibleAssetCount; }
+	int32 GetHiddenAssetCount() const { return Assets.Num() - VisibleAssetCount; }
+
+	void SetVisibleAssetCount(int32 InVisibleAssetCount) { check(VisibleAssetCount <= Assets.Num()); VisibleAssetCount = InVisibleAssetCount; }
+
+	void AddAsset(const FAssetTableRow& AssetRow) { Assets.Add(AssetRow); }
+
+	void ClearAllData()
+	{
+		Assets.Empty();
+		VisibleAssetCount = 0;
+	}
+
 private:
 	void AddDefaultColumns();
 
 private:
 	TArray<FAssetTableRow> Assets;
+	int32 VisibleAssetCount = 0;
 };
diff --git a/Engine/Plugins/Editor/AssetManagerEditor/Source/AssetManagerEditor/Private/TreeView/AssetTreeNode.h b/Engine/Plugins/Editor/AssetManagerEditor/Source/AssetManagerEditor/Private/TreeView/AssetTreeNode.h
--- a/Engine/Plugins/Editor/AssetManagerEditor/Source/AssetManagerEditor/Private/TreeView/AssetTreeNode.h
+++ b/Engine/Plugins/Editor/AssetManagerEditor/Source/AssetManagerEditor/Private/TreeView/AssetTreeNode.h
@@ -29,6 +29,11 @@ public:
 	FAssetTable& GetAssetTableChecked() const { return *AssetTable; }
 	const FAssetTableRow& GetAssetChecked() const { return AssetTable->GetAssetChecked(RowId.RowIndex); }
 
+	virtual FLinearColor GetColor() const override
+	{
+		return GetAssetChecked().GetColor();
+	}
+
 private:
 	FAssetTable* AssetTable;
 };
diff --git a/Engine/Plugins/Editor/AssetManagerEditor/Source/AssetManagerEditor/Private/TreeView/SAssetTableTreeView.cpp b/Engine/Plugins/Editor/AssetManagerEditor/Source/AssetManagerEditor/Private/TreeView/SAssetTableTreeView.cpp
--- a/Engine/Plugins/Editor/AssetManagerEditor/Source/AssetManagerEditor/Private/TreeView/SAssetTableTreeView.cpp
+++ b/Engine/Plugins/Editor/AssetManagerEditor/Source/AssetManagerEditor/Private/TreeView/SAssetTableTreeView.cpp
@@ -15,6 +15,7 @@
 #include "HAL/PlatformFileManager.h"
 #include "Logging/MessageLog.h"
 
+#include "Insights/Common/InsightsStyle.h"
 #include "Insights/Common/Log.h"
 #include "Insights/Common/Stopwatch.h"
 #include "Insights/Filter/ViewModels/FilterConfigurator.h"
@@ -23,6 +24,7 @@
 #include "Insights/Table/ViewModels/TreeNodeGrouping.h"
 #include "TreeView/AssetTable.h"
 #include "TreeView/AssetTreeNode.h"
+#include "TreeView/AssetDependencyGrouping.h"
 
 #include <limits>
 #include <memory>
@@ -92,29 +94,33 @@ void SAssetTableTreeView::RebuildTree(bool bResync)
 
 	TSharedPtr<FAssetTable> AssetTable = GetAssetTable();
 	const TArray<FAssetTableRow>& Assets = AssetTable->GetAssets();
-	const int32 TotalAssetCount = Assets.Num();
+	const int32 VisibleAssetCount = AssetTable->GetVisibleAssetCount();
 
-	if (TotalAssetCount != PreviousNodeCount)
+	if (VisibleAssetCount != PreviousNodeCount)
 	{
-		UE_LOG(LogInsights, Log, TEXT("[AssetTree] Creating nodes (%d nodes --> %d assets)..."), PreviousNodeCount, TotalAssetCount);
+		UE_LOG(LogInsights, Log, TEXT("[AssetTree] Creating nodes (%d nodes --> %d assets)..."), PreviousNodeCount, VisibleAssetCount);
 
-		if (TotalAssetCount < PreviousNodeCount)
+		if (VisibleAssetCount < PreviousNodeCount)
 		{
 			TableTreeNodes.Empty();
 		}
-		TableTreeNodes.Reserve(TotalAssetCount);
+		TableTreeNodes.Reserve(VisibleAssetCount);
 
 		//const FName BaseNodeName(TEXT("asset"));
-		for (int32 AssetIndex = TableTreeNodes.Num(); AssetIndex < TotalAssetCount; ++AssetIndex)
+		for (int32 AssetIndex = TableTreeNodes.Num(); AssetIndex < VisibleAssetCount; ++AssetIndex)
 		{
 			const FAssetTableRow* Asset = AssetTable->GetAsset(AssetIndex);
 
 			//FName NodeName(BaseNodeName, AssetIndex + 1);
 			FName NodeName(Asset->GetName());
 			FAssetTreeNodePtr NodePtr = MakeShared<FAssetTreeNode>(NodeName, AssetTable, AssetIndex);
+			if (AssetIndex % 3) // TODO
+				NodePtr->SetIcon(UE::Insights::FInsightsStyle::GetBrush("Icons.Asset.TreeItem"));
+			else if (AssetIndex % 7) // TODO
+				NodePtr->SetIcon(UE::Insights::FInsightsStyle::GetBrush("Icons.Package.TreeItem"));
 			TableTreeNodes.Add(NodePtr);
 		}
-		ensure(TableTreeNodes.Num() == TotalAssetCount);
+		ensure(TableTreeNodes.Num() == VisibleAssetCount);
 	}
 
 	SyncStopwatch.Stop();
@@ -252,20 +258,21 @@ void SAssetTableTreeView::InitAvailableViewPresets()
 		}
 		virtual void GetColumnConfigSet(TArray<UE::Insights::FTableColumnConfig>& InOutConfigSet) const override
 		{
-			InOutConfigSet.Add({ UE::Insights::FTable::GetHierarchyColumnId(),         true, 200.0f });
+			InOutConfigSet.Add({ UE::Insights::FTable::GetHierarchyColumnId(),     true, 200.0f });
 			InOutConfigSet.Add({ FAssetTableColumns::CountColumnId,                true, 100.0f });
 			InOutConfigSet.Add({ FAssetTableColumns::NameColumnId,                 true, 200.0f });
 			InOutConfigSet.Add({ FAssetTableColumns::TypeColumnId,                 true, 200.0f });
 			InOutConfigSet.Add({ FAssetTableColumns::PathColumnId,                 true, 400.0f });
 			InOutConfigSet.Add({ FAssetTableColumns::PrimaryTypeColumnId,          true, 200.0f });
 			InOutConfigSet.Add({ FAssetTableColumns::PrimaryNameColumnId,          true, 200.0f });
-			InOutConfigSet.Add({ FAssetTableColumns::ManagedDiskSizeColumnId,      true, 100.0f });
-			InOutConfigSet.Add({ FAssetTableColumns::DiskSizeColumnId,             true, 100.0f });
+			//InOutConfigSet.Add({ FAssetTableColumns::ManagedDiskSizeColumnId,      true, 100.0f });
+			//InOutConfigSet.Add({ FAssetTableColumns::DiskSizeColumnId,             true, 100.0f });
 			InOutConfigSet.Add({ FAssetTableColumns::StagedCompressedSizeColumnId, true, 100.0f });
 			InOutConfigSet.Add({ FAssetTableColumns::TotalUsageCountColumnId,      true, 100.0f });
-			InOutConfigSet.Add({ FAssetTableColumns::CookRuleColumnId,             true, 200.0f });
-			InOutConfigSet.Add({ FAssetTableColumns::ChunksColumnId,               true, 200.0f });
+			//InOutConfigSet.Add({ FAssetTableColumns::CookRuleColumnId,             true, 200.0f });
+			//InOutConfigSet.Add({ FAssetTableColumns::ChunksColumnId,               true, 200.0f });
 			InOutConfigSet.Add({ FAssetTableColumns::NativeClassColumnId,          true, 200.0f });
+			InOutConfigSet.Add({ FAssetTableColumns::GameFeaturePluginColumnId,    true, 200.0f });
 		}
 	};
 	AvailableViewPresets.Add(MakeShared<FDefaultViewPreset>());
@@ -312,20 +319,21 @@ void SAssetTableTreeView::InitAvailableViewPresets()
 		}
 		virtual void GetColumnConfigSet(TArray<UE::Insights::FTableColumnConfig>& InOutConfigSet) const override
 		{
-			InOutConfigSet.Add({ UE::Insights::FTable::GetHierarchyColumnId(),         true, 400.0f });
+			InOutConfigSet.Add({ UE::Insights::FTable::GetHierarchyColumnId(),     true, 400.0f });
 			InOutConfigSet.Add({ FAssetTableColumns::CountColumnId,                true, 100.0f });
 			InOutConfigSet.Add({ FAssetTableColumns::NameColumnId,                 true, 200.0f });
 			InOutConfigSet.Add({ FAssetTableColumns::TypeColumnId,                 true, 200.0f });
 			//InOutConfigSet.Add({ FAssetTableColumns::PathColumnId,                 true, 400.0f });
 			InOutConfigSet.Add({ FAssetTableColumns::PrimaryTypeColumnId,          true, 200.0f });
 			InOutConfigSet.Add({ FAssetTableColumns::PrimaryNameColumnId,          true, 200.0f });
-			InOutConfigSet.Add({ FAssetTableColumns::ManagedDiskSizeColumnId,      true, 100.0f });
-			InOutConfigSet.Add({ FAssetTableColumns::DiskSizeColumnId,             true, 100.0f });
+			//InOutConfigSet.Add({ FAssetTableColumns::ManagedDiskSizeColumnId,      true, 100.0f });
+			//InOutConfigSet.Add({ FAssetTableColumns::DiskSizeColumnId,             true, 100.0f });
 			InOutConfigSet.Add({ FAssetTableColumns::StagedCompressedSizeColumnId, true, 100.0f });
 			InOutConfigSet.Add({ FAssetTableColumns::TotalUsageCountColumnId,      true, 100.0f });
-			InOutConfigSet.Add({ FAssetTableColumns::CookRuleColumnId,             true, 200.0f });
-			InOutConfigSet.Add({ FAssetTableColumns::ChunksColumnId,               true, 200.0f });
+			//InOutConfigSet.Add({ FAssetTableColumns::CookRuleColumnId,             true, 200.0f });
+			//InOutConfigSet.Add({ FAssetTableColumns::ChunksColumnId,               true, 200.0f });
 			InOutConfigSet.Add({ FAssetTableColumns::NativeClassColumnId,          true, 200.0f });
+			InOutConfigSet.Add({ FAssetTableColumns::GameFeaturePluginColumnId,    true, 200.0f });
 		}
 	};
 	AvailableViewPresets.Add(MakeShared<FAssetPathViewPreset>());
@@ -383,20 +391,21 @@ void SAssetTableTreeView::InitAvailableViewPresets()
 		}
 		virtual void GetColumnConfigSet(TArray<UE::Insights::FTableColumnConfig>& InOutConfigSet) const override
 		{
-			InOutConfigSet.Add({ UE::Insights::FTable::GetHierarchyColumnId(),         true, 300.0f });
+			InOutConfigSet.Add({ UE::Insights::FTable::GetHierarchyColumnId(),     true, 300.0f });
 			InOutConfigSet.Add({ FAssetTableColumns::CountColumnId,                true, 100.0f });
 			InOutConfigSet.Add({ FAssetTableColumns::NameColumnId,                 true, 200.0f });
 			InOutConfigSet.Add({ FAssetTableColumns::TypeColumnId,                 true, 200.0f });
 			InOutConfigSet.Add({ FAssetTableColumns::PathColumnId,                 true, 400.0f });
 			//InOutConfigSet.Add({ FAssetTableColumns::PrimaryTypeColumnId,          true, 200.0f });
 			//InOutConfigSet.Add({ FAssetTableColumns::PrimaryNameColumnId,          true, 200.0f });
-			InOutConfigSet.Add({ FAssetTableColumns::ManagedDiskSizeColumnId,      true, 100.0f });
-			InOutConfigSet.Add({ FAssetTableColumns::DiskSizeColumnId,             true, 100.0f });
+			//InOutConfigSet.Add({ FAssetTableColumns::ManagedDiskSizeColumnId,      true, 100.0f });
+			//InOutConfigSet.Add({ FAssetTableColumns::DiskSizeColumnId,             true, 100.0f });
 			InOutConfigSet.Add({ FAssetTableColumns::StagedCompressedSizeColumnId, true, 100.0f });
 			InOutConfigSet.Add({ FAssetTableColumns::TotalUsageCountColumnId,      true, 100.0f });
-			InOutConfigSet.Add({ FAssetTableColumns::CookRuleColumnId,             true, 200.0f });
-			InOutConfigSet.Add({ FAssetTableColumns::ChunksColumnId,               true, 200.0f });
+			//InOutConfigSet.Add({ FAssetTableColumns::CookRuleColumnId,             true, 200.0f });
+			//InOutConfigSet.Add({ FAssetTableColumns::ChunksColumnId,               true, 200.0f });
 			InOutConfigSet.Add({ FAssetTableColumns::NativeClassColumnId,          true, 200.0f });
+			InOutConfigSet.Add({ FAssetTableColumns::GameFeaturePluginColumnId,    true, 200.0f });
 		}
 	};
 	AvailableViewPresets.Add(MakeShared<FPrimaryAssetViewPreset>());
@@ -452,8 +461,7 @@ void SAssetTableTreeView::InternalCreateGroupings()
 			else if (Grouping->Is<UE::Insights::FTreeNodeGroupingByPathBreakdown>())
 			{
 				const FName ColumnId = Grouping->As<UE::Insights::FTreeNodeGroupingByPathBreakdown>().GetColumnId();
-				if (ColumnId == FAssetTableColumns::TypeColumnId ||
-					ColumnId == FAssetTableColumns::NativeClassColumnId)
+				if (ColumnId != FAssetTableColumns::PathColumnId)
 				{
 					return true;
 				}
@@ -461,7 +469,10 @@ void SAssetTableTreeView::InternalCreateGroupings()
 			return false;
 		});
 
+	// Add custom groupings...
 	int32 Index = 1; // after the Flat ("All") grouping
+
+	AvailableGroupings.Insert(MakeShared<FAssetDependencyGrouping>(), Index++);
 }
 
 ////////////////////////////////////////////////////////////////////////////////////////////////////
diff --git a/Engine/Plugins/Editor/AssetManagerEditor/Source/AssetManagerEditor/Public/AssetManagerEditorModule.h b/Engine/Plugins/Editor/AssetManagerEditor/Source/AssetManagerEditor/Public/AssetManagerEditorModule.h
--- a/Engine/Plugins/Editor/AssetManagerEditor/Source/AssetManagerEditor/Public/AssetManagerEditorModule.h
+++ b/Engine/Plugins/Editor/AssetManagerEditor/Source/AssetManagerEditor/Public/AssetManagerEditorModule.h
@@ -383,6 +383,7 @@ public:
 	static const FName ChunksName;
 	static const FName StageChunkSizeName;
 	static const FName StageChunkCompressedSizeName;
+	static const FName GameFeaturePluginsName;
 
 	/** Gets the value of a "virtual" column for an asset data, this will query the AssetManager for you and takes current platform into account. Returns true and sets out parameter if found */
 	virtual bool GetStringValueForCustomColumn(const FAssetData& AssetData, FName ColumnName, FString& OutValue) = 0;
