AssetManagerEditor: Updated LOCTEXT_NAMESPACE naming convention in Insights' Tree View implementation.

diff --git a/Engine/Plugins/Editor/AssetManagerEditor/Source/AssetManagerEditor/Private/Insights/Table/ViewModels/TableCellValueFormatter.cpp b/Engine/Plugins/Editor/AssetManagerEditor/Source/AssetManagerEditor/Private/Insights/Table/ViewModels/TableCellValueFormatter.cpp
--- a/Engine/Plugins/Editor/AssetManagerEditor/Source/AssetManagerEditor/Private/Insights/Table/ViewModels/TableCellValueFormatter.cpp
+++ b/Engine/Plugins/Editor/AssetManagerEditor/Source/AssetManagerEditor/Private/Insights/Table/ViewModels/TableCellValueFormatter.cpp
@@ -11,7 +11,7 @@
 #include "Insights/Table/ViewModels/TableCellValueGetter.h"
 #include "Insights/Table/ViewModels/TableColumn.h"
 
-#define LOCTEXT_NAMESPACE "TableCellValueFormatter"
+#define LOCTEXT_NAMESPACE "UE::Insights::FTableCellValueFormatter"
 
 namespace UE
 {
diff --git a/Engine/Plugins/Editor/AssetManagerEditor/Source/AssetManagerEditor/Private/Insights/Table/ViewModels/TableCellValueGetter.cpp b/Engine/Plugins/Editor/AssetManagerEditor/Source/AssetManagerEditor/Private/Insights/Table/ViewModels/TableCellValueGetter.cpp
--- a/Engine/Plugins/Editor/AssetManagerEditor/Source/AssetManagerEditor/Private/Insights/Table/ViewModels/TableCellValueGetter.cpp
+++ b/Engine/Plugins/Editor/AssetManagerEditor/Source/AssetManagerEditor/Private/Insights/Table/ViewModels/TableCellValueGetter.cpp
@@ -4,7 +4,7 @@
 
 #include "Insights/Table/ViewModels/BaseTreeNode.h"
 
-#define LOCTEXT_NAMESPACE "TableCellValueGetter"
+#define LOCTEXT_NAMESPACE "UE::Insights::FTableCellValueGetter"
 
 namespace UE
 {
diff --git a/Engine/Plugins/Editor/AssetManagerEditor/Source/AssetManagerEditor/Private/Insights/Table/ViewModels/TableCellValueSorter.cpp b/Engine/Plugins/Editor/AssetManagerEditor/Source/AssetManagerEditor/Private/Insights/Table/ViewModels/TableCellValueSorter.cpp
--- a/Engine/Plugins/Editor/AssetManagerEditor/Source/AssetManagerEditor/Private/Insights/Table/ViewModels/TableCellValueSorter.cpp
+++ b/Engine/Plugins/Editor/AssetManagerEditor/Source/AssetManagerEditor/Private/Insights/Table/ViewModels/TableCellValueSorter.cpp
@@ -4,7 +4,7 @@
 
 #include "Insights/Table/ViewModels/BaseTreeNode.h"
 
-#define LOCTEXT_NAMESPACE "TableCellValueSorter"
+#define LOCTEXT_NAMESPACE "UE::Insights::FTableCellValueSorter"
 
 // Default pre-sorting (group nodes sorts above leaf nodes)
 #define INSIGHTS_DEFAULT_PRESORTING_NODES(A, B) \
diff --git a/Engine/Plugins/Editor/AssetManagerEditor/Source/AssetManagerEditor/Private/Insights/Table/ViewModels/TableColumn.cpp b/Engine/Plugins/Editor/AssetManagerEditor/Source/AssetManagerEditor/Private/Insights/Table/ViewModels/TableColumn.cpp
--- a/Engine/Plugins/Editor/AssetManagerEditor/Source/AssetManagerEditor/Private/Insights/Table/ViewModels/TableColumn.cpp
+++ b/Engine/Plugins/Editor/AssetManagerEditor/Source/AssetManagerEditor/Private/Insights/Table/ViewModels/TableColumn.cpp
@@ -6,7 +6,7 @@
 #include "Insights/Table/ViewModels/TableCellValueGetter.h"
 #include "Insights/Table/ViewModels/TableCellValueSorter.h"
 
-#define LOCTEXT_NAMESPACE "TableColumn"
+#define LOCTEXT_NAMESPACE "UE::Insights::FTableColumn"
 
 namespace UE
 {
diff --git a/Engine/Plugins/Editor/AssetManagerEditor/Source/AssetManagerEditor/Private/Insights/Table/ViewModels/TreeNodeGrouping.cpp b/Engine/Plugins/Editor/AssetManagerEditor/Source/AssetManagerEditor/Private/Insights/Table/ViewModels/TreeNodeGrouping.cpp
--- a/Engine/Plugins/Editor/AssetManagerEditor/Source/AssetManagerEditor/Private/Insights/Table/ViewModels/TreeNodeGrouping.cpp
+++ b/Engine/Plugins/Editor/AssetManagerEditor/Source/AssetManagerEditor/Private/Insights/Table/ViewModels/TreeNodeGrouping.cpp
@@ -5,7 +5,7 @@
 #include "Insights/Common/AsyncOperationProgress.h"
 #include "Insights/Table/ViewModels/TableTreeNode.h"
 
-#define LOCTEXT_NAMESPACE "Insights_TreeNode"
+#define LOCTEXT_NAMESPACE "UE::Insights::FTreeNodeGrouping"
 
 namespace UE
 {
