AssetManagerEditor:
  - Updated Insights' Tree View implementations (support for icons and color tint for tree nodes; increased font size; added FCustomTableTreeNode).
  - Added colors for special group nodes created during "dependencies" grouping.

diff --git a/Engine/Plugins/Editor/AssetManagerEditor/Source/AssetManagerEditor/Private/Insights/Common/InsightsStyle.cpp b/Engine/Plugins/Editor/AssetManagerEditor/Source/AssetManagerEditor/Private/Insights/Common/InsightsStyle.cpp
--- a/Engine/Plugins/Editor/AssetManagerEditor/Source/AssetManagerEditor/Private/Insights/Common/InsightsStyle.cpp
+++ b/Engine/Plugins/Editor/AssetManagerEditor/Source/AssetManagerEditor/Private/Insights/Common/InsightsStyle.cpp
@@ -141,7 +141,7 @@ void FInsightsStyle::FStyle::Initialize()
 	Set("Icons.ClassicFilterConfig", new IMAGE_BRUSH_SVG("FilterConfig", Icon16x16));
 
 	Set("Icons.SortBy", new CORE_IMAGE_BRUSH("Icons/Profiler/profiler_SortBy_32x", Icon16x16));
-	//Set("Icons.SortUp", new CORE_IMAGE_BRUSH_SVG("Starship/Common/SortUp", Icon16x16));	//-> use FAppStyle "Icons.SortUp"
+	//Set("Icons.SortUp", new CORE_IMAGE_BRUSH_SVG("Starship/Common/SortUp", Icon16x16));		//-> use FAppStyle "Icons.SortUp"
 	//Set("Icons.SortDown", new CORE_IMAGE_BRUSH_SVG("Starship/Common/SortDown", Icon16x16));	//-> use FAppStyle "Icons.SortDown"
 
 	Set("Icons.ViewColumn", new CORE_IMAGE_BRUSH("Icons/Profiler/profiler_ViewColumn_32x", Icon16x16));
@@ -180,6 +180,14 @@ void FInsightsStyle::FStyle::Initialize()
 		.SetShadowColorAndOpacity(FLinearColor(0.0f, 0.0f, 0.0f, 0.8f))
 	);
 
+	Set("TreeTable.NameText", FTextBlockStyle(NormalText)
+		.SetColorAndOpacity(FLinearColor::White)
+	);
+
+	Set("TreeTable.NormalText", FTextBlockStyle(NormalText)
+		.SetColorAndOpacity(FLinearColor::White)
+	);
+
 	//////////////////////////////////////////////////
 
 	// NormalEditableTextBox && SearchBox
diff --git a/Engine/Plugins/Editor/AssetManagerEditor/Source/AssetManagerEditor/Private/Insights/Table/ViewModels/BaseTreeNode.cpp b/Engine/Plugins/Editor/AssetManagerEditor/Source/AssetManagerEditor/Private/Insights/Table/ViewModels/BaseTreeNode.cpp
--- a/Engine/Plugins/Editor/AssetManagerEditor/Source/AssetManagerEditor/Private/Insights/Table/ViewModels/BaseTreeNode.cpp
+++ b/Engine/Plugins/Editor/AssetManagerEditor/Source/AssetManagerEditor/Private/Insights/Table/ViewModels/BaseTreeNode.cpp
@@ -2,6 +2,9 @@
 
 #include "BaseTreeNode.h"
 
+#include "Styling/SlateBrush.h"
+#include "Styling/StyleColors.h"
+
 // Insights
 #include "Insights/Common/InsightsStyle.h"
 #include "Insights/Table/ViewModels/TableCellValueSorter.h"
@@ -55,23 +58,30 @@ bool FBaseTreeNode::HasExtraDisplayName() const
 
 ////////////////////////////////////////////////////////////////////////////////////////////////////
 
-void FBaseTreeNode::SetDefaultIcon()
+const FSlateBrush* FBaseTreeNode::GetDefaultIcon(bool bIsGroupNode)
 {
-	if (IsGroup())
+	if (bIsGroupNode)
 	{
-		IconBrush = FInsightsStyle::GetBrush("Icons.Group.TreeItem");
+		return FInsightsStyle::GetBrush("Icons.Group.TreeItem");
 	}
 	else
 	{
-		IconBrush = FInsightsStyle::GetBrush("Icons.Leaf.TreeItem");
+		return FInsightsStyle::GetBrush("Icons.Leaf.TreeItem");
 	}
 }
 
 ////////////////////////////////////////////////////////////////////////////////////////////////////
 
-FLinearColor FBaseTreeNode::GetColor() const
+FLinearColor FBaseTreeNode::GetDefaultColor(bool bIsGroupNode)
 {
-	return FLinearColor(1.0f, 1.0f, 1.0f, 1.0f);
+	if (bIsGroupNode)
+	{
+		return FLinearColor(1.0f, 0.9f, 0.6f, 1.0f);
+	}
+	else
+	{
+		return FLinearColor(1.0f, 1.0f, 1.0f, 1.0f);
+	}
 }
 
 ////////////////////////////////////////////////////////////////////////////////////////////////////
diff --git a/Engine/Plugins/Editor/AssetManagerEditor/Source/AssetManagerEditor/Private/Insights/Table/ViewModels/BaseTreeNode.h b/Engine/Plugins/Editor/AssetManagerEditor/Source/AssetManagerEditor/Private/Insights/Table/ViewModels/BaseTreeNode.h
--- a/Engine/Plugins/Editor/AssetManagerEditor/Source/AssetManagerEditor/Private/Insights/Table/ViewModels/BaseTreeNode.h
+++ b/Engine/Plugins/Editor/AssetManagerEditor/Source/AssetManagerEditor/Private/Insights/Table/ViewModels/BaseTreeNode.h
@@ -3,6 +3,7 @@
 #pragma once
 
 #include "CoreMinimal.h"
+#include "Math/Color.h"
 
 #include "Insights/Common/SimpleRtti.h"
 
@@ -58,13 +59,11 @@ protected:
 
 public:
 	/** Initialization constructor for the node. */
-	FBaseTreeNode(const FName InName, bool bInIsGroup)
+	explicit FBaseTreeNode(const FName InName, bool bInIsGroup)
 		: DefaultSortOrder(0)
 		, Name(InName)
-		, IconBrush(nullptr)
 		, GroupData(bInIsGroup ? new FGroupNodeData() : &DefaultGroupData)
 	{
-		SetDefaultIcon();
 	}
 
 	virtual ~FBaseTreeNode()
@@ -125,30 +124,30 @@ public:
 	}
 
 	/**
-	 * @return a brush icon for this node.
+	 * @return the default icon for a group/leaf node.
 	 */
-	const FSlateBrush* GetIcon() const
-	{
-		return IconBrush;
-	}
+	static const FSlateBrush* GetDefaultIcon(bool bIsGroupNode);
 
 	/**
-	 * Sets an icon brush for this node.
+	 * @return a brush icon for this node.
 	 */
-	void SetIcon(const FSlateBrush* InIconBrush)
+	virtual const FSlateBrush* GetIcon() const
 	{
-		IconBrush = InIconBrush;
+		return GetDefaultIcon(IsGroup());
 	}
 
 	/**
-	 * Sets the default icon brush for this node.
+	 * @return the default color tint for a group/leaf node.
 	 */
-	void SetDefaultIcon();
+	static FLinearColor GetDefaultColor(bool bIsGroupNode);
 
 	/**
-	 * @return the color tint for icon and name text.
+	 * @return the color tint to be used for the icon and the name text of this node.
 	 */
-	virtual FLinearColor GetColor() const;
+	virtual FLinearColor GetColor() const
+	{
+		return GetDefaultColor(IsGroup());
+	}
 
 	/**
 	 * @return a pointer to a data context for this node.
@@ -322,9 +321,6 @@ private:
 	/** The name of this node. */
 	const FName Name;
 
-	/** The icon of this node. */
-	const FSlateBrush* IconBrush;
-
 	/** A weak pointer to the group/parent of this node. */
 	FBaseTreeNodeWeak GroupPtr;
 
diff --git a/Engine/Plugins/Editor/AssetManagerEditor/Source/AssetManagerEditor/Private/Insights/Table/ViewModels/TableTreeNode.cpp b/Engine/Plugins/Editor/AssetManagerEditor/Source/AssetManagerEditor/Private/Insights/Table/ViewModels/TableTreeNode.cpp
--- a/Engine/Plugins/Editor/AssetManagerEditor/Source/AssetManagerEditor/Private/Insights/Table/ViewModels/TableTreeNode.cpp
+++ b/Engine/Plugins/Editor/AssetManagerEditor/Source/AssetManagerEditor/Private/Insights/Table/ViewModels/TableTreeNode.cpp
@@ -10,6 +10,7 @@ namespace Insights
 {
 
 INSIGHTS_IMPLEMENT_RTTI(FTableTreeNode)
+INSIGHTS_IMPLEMENT_RTTI(FCustomTableTreeNode)
 
 } // namespace Insights
 } // namespace UE
diff --git a/Engine/Plugins/Editor/AssetManagerEditor/Source/AssetManagerEditor/Private/Insights/Table/ViewModels/TableTreeNode.h b/Engine/Plugins/Editor/AssetManagerEditor/Source/AssetManagerEditor/Private/Insights/Table/ViewModels/TableTreeNode.h
--- a/Engine/Plugins/Editor/AssetManagerEditor/Source/AssetManagerEditor/Private/Insights/Table/ViewModels/TableTreeNode.h
+++ b/Engine/Plugins/Editor/AssetManagerEditor/Source/AssetManagerEditor/Private/Insights/Table/ViewModels/TableTreeNode.h
@@ -53,7 +53,7 @@ class FTableTreeNode : public FBaseTreeNode
 
 public:
 	/** Initialization constructor for a table record node. */
-	FTableTreeNode(const FName InName, TWeakPtr<FTable> InParentTable, int32 InRowIndex)
+	explicit FTableTreeNode(const FName InName, TWeakPtr<FTable> InParentTable, int32 InRowIndex)
 		: FBaseTreeNode(InName, false)
 		, ParentTable(InParentTable)
 		, RowId(InRowIndex)
@@ -62,7 +62,7 @@ public:
 	}
 
 	/** Initialization constructor for a group node. */
-	FTableTreeNode(const FName InGroupName, TWeakPtr<FTable> InParentTable)
+	explicit FTableTreeNode(const FName InGroupName, TWeakPtr<FTable> InParentTable)
 		: FBaseTreeNode(InGroupName, true)
 		, ParentTable(InParentTable)
 		, RowId(FTableRowId::InvalidRowIndex)
@@ -81,6 +81,7 @@ public:
 
 	void InitAggregatedValues()
 	{
+		check(IsGroup());
 		if (!AggregatedValues)
 		{
 			AggregatedValues = new TMap<FName, FTableCellValue>();
@@ -119,5 +120,64 @@ private:
 
 ////////////////////////////////////////////////////////////////////////////////////////////////////
 
+class FCustomTableTreeNode : public FTableTreeNode
+{
+	INSIGHTS_DECLARE_RTTI(FCustomTableTreeNode, FTableTreeNode)
+
+public:
+	/** Initialization constructor for a table record node. */
+	explicit FCustomTableTreeNode(const FName InName, TWeakPtr<FTable> InParentTable, int32 InRowIndex, const FSlateBrush* InIconBrush, FLinearColor InColor)
+		: FTableTreeNode(InName, InParentTable)
+		, IconBrush(InIconBrush)
+		, Color(InColor)
+	{
+	}
+
+	/** Initialization constructor for the group node. */
+	explicit FCustomTableTreeNode(const FName InName, TWeakPtr<FTable> InParentTable, const FSlateBrush* InIconBrush, FLinearColor InColor)
+		: FTableTreeNode(InName, InParentTable)
+		, IconBrush(InIconBrush)
+		, Color(InColor)
+	{
+	}
+
+	virtual ~FCustomTableTreeNode()
+	{
+	}
+
+	/**
+	 * @return a brush icon for this node.
+	 */
+	virtual const FSlateBrush* GetIcon() const override
+	{
+		return IconBrush;
+	}
+
+	/**
+	 * Sets an icon brush for this node.
+	 */
+	void SetIcon(const FSlateBrush* InIconBrush)
+	{
+		IconBrush = InIconBrush;
+	}
+
+	/**
+	 * @return the color tint for icon and name text.
+	 */
+	virtual FLinearColor GetColor() const override
+	{
+		return Color;
+	}
+
+private:
+	/** The icon of this node. */
+	const FSlateBrush* IconBrush;
+
+	/** The color tint of this node. */
+	FLinearColor Color;
+};
+
+////////////////////////////////////////////////////////////////////////////////////////////////////
+
 } // namespace Insights
 } // namespace UE
diff --git a/Engine/Plugins/Editor/AssetManagerEditor/Source/AssetManagerEditor/Private/Insights/Table/Widgets/STableTreeViewCell.cpp b/Engine/Plugins/Editor/AssetManagerEditor/Source/AssetManagerEditor/Private/Insights/Table/Widgets/STableTreeViewCell.cpp
--- a/Engine/Plugins/Editor/AssetManagerEditor/Source/AssetManagerEditor/Private/Insights/Table/Widgets/STableTreeViewCell.cpp
+++ b/Engine/Plugins/Editor/AssetManagerEditor/Source/AssetManagerEditor/Private/Insights/Table/Widgets/STableTreeViewCell.cpp
@@ -3,6 +3,7 @@
 #include "STableTreeViewCell.h"
 
 #include "SlateOptMacros.h"
+#include "Styling/StyleColors.h"
 #include "Widgets/Images/SImage.h"
 #include "Widgets/Input/SButton.h"
 #include "Widgets/SBoxPanel.h"
@@ -28,8 +29,10 @@ namespace Insights
 
 BEGIN_SLATE_FUNCTION_BUILD_OPTIMIZATION
 
-void STableTreeViewCell::Construct(const FArguments& InArgs, const TSharedRef<ITableRow>& TableRow)
+void STableTreeViewCell::Construct(const FArguments& InArgs, const TSharedRef<ITableRow>& InTableRow)
 {
+	TableRow = InTableRow;
+
 	TablePtr = InArgs._TablePtr;
 	ColumnPtr = InArgs._ColumnPtr;
 	TableTreeNodePtr = InArgs._TableTreeNodePtr;
@@ -42,28 +45,31 @@ void STableTreeViewCell::Construct(const FArguments& InArgs, const TSharedRef<IT
 
 	ChildSlot
 	[
-		GenerateWidgetForColumn(InArgs, TableRow)
+		GenerateWidgetForColumn(InArgs)
 	];
 }
 
 ////////////////////////////////////////////////////////////////////////////////////////////////////
 
-TSharedRef<SWidget> STableTreeViewCell::GenerateWidgetForColumn(const FArguments& InArgs, const TSharedRef<ITableRow>& TableRow)
+TSharedRef<SWidget> STableTreeViewCell::GenerateWidgetForColumn(const FArguments& InArgs)
 {
 	if (InArgs._IsNameColumn)
 	{
-		return GenerateWidgetForNameColumn(InArgs, TableRow);
+		return GenerateWidgetForNameColumn(InArgs);
 	}
 	else
 	{
-		return GenerateWidgetForTableColumn(InArgs, TableRow);
+		return GenerateWidgetForTableColumn(InArgs);
 	}
 }
 
 ////////////////////////////////////////////////////////////////////////////////////////////////////
 
-TSharedRef<SWidget> STableTreeViewCell::GenerateWidgetForNameColumn(const FArguments& InArgs, const TSharedRef<ITableRow>& TableRow)
+TSharedRef<SWidget> STableTreeViewCell::GenerateWidgetForNameColumn(const FArguments& InArgs)
 {
+	TSharedPtr<STableTreeViewRow> Row = StaticCastSharedPtr<STableTreeViewRow, ITableRow>(TableRow);
+	TSharedPtr<IToolTip> RowToolTip = Row->GetRowToolTip();
+
 	return
 		SNew(SHorizontalBox)
 
@@ -81,64 +87,48 @@ TSharedRef<SWidget> STableTreeViewCell::GenerateWidgetForNameColumn(const FArgum
 		.HAlign(HAlign_Center)
 		.VAlign(VAlign_Center)
 		[
-			SNew(SOverlay)
-			+ SOverlay::Slot()
-			[
-				SNew(SImage)
-				.Image(this, &STableTreeViewCell::GetIcon)
-				.ColorAndOpacity(this, &STableTreeViewCell::GetIconColorAndOpacity)
-			]
-			+ SOverlay::Slot()
-			[
-				SNew(SImage)
-				.Visibility(this, &STableTreeViewCell::GetHintIconVisibility)
-				.Image(FInsightsStyle::GetBrush("Icons.Hint.TreeItem"))
-				.ToolTip(GetRowToolTip(TableRow))
-			]
+			SNew(SImage)
+			.Image(this, &STableTreeViewCell::GetIcon)
+			.ColorAndOpacity(this, &STableTreeViewCell::GetIconColorAndOpacity)
+			.ToolTip(RowToolTip)
 		]
 
 		// Name
 		+ SHorizontalBox::Slot()
 		.AutoWidth()
-		.VAlign(VAlign_Top)
 		.HAlign(HAlign_Left)
-		.Padding(FMargin(2.0f, 1.5f, 2.0f, 0.0f))
+		.VAlign(VAlign_Center)
+		.Padding(FMargin(2.0f, 0.0f, 2.0f, 0.0f))
 		[
 			SNew(STextBlock)
 			.Text(this, &STableTreeViewCell::GetDisplayName)
 			.HighlightText(InArgs._HighlightText)
-			.TextStyle(FInsightsStyle::Get(), TEXT("TreeTable.Tooltip"))
-			.ColorAndOpacity(this, &STableTreeViewCell::GetColorAndOpacity)
+			.TextStyle(FInsightsStyle::Get(), TEXT("TreeTable.NameText"))
+			.ColorAndOpacity(this, &STableTreeViewCell::GetDisplayNameColorAndOpacity)
 			.ShadowColorAndOpacity(this, &STableTreeViewCell::GetShadowColorAndOpacity)
+			.ToolTip(RowToolTip)
 		]
 
 		// Name Suffix
 		+ SHorizontalBox::Slot()
 		.AutoWidth()
-		.VAlign(VAlign_Top)
 		.HAlign(HAlign_Left)
-		.Padding(FMargin(2.0f, 1.5f, 2.0f, 0.0f))
+		.VAlign(VAlign_Center)
+		.Padding(FMargin(2.0f, 0.0f, 2.0f, 0.0f))
 		[
 			SNew(STextBlock)
 			.Visibility(this, &STableTreeViewCell::HasExtraDisplayName)
 			.Text(this, &STableTreeViewCell::GetExtraDisplayName)
-			.TextStyle(FInsightsStyle::Get(), TEXT("TreeTable.Tooltip"))
-			.ColorAndOpacity(this, &STableTreeViewCell::GetExtraColorAndOpacity)
+			.TextStyle(FInsightsStyle::Get(), TEXT("TreeTable.NameText"))
+			.ColorAndOpacity(this, &STableTreeViewCell::GetExtraDisplayNameColorAndOpacity)
 			.ShadowColorAndOpacity(this, &STableTreeViewCell::GetShadowColorAndOpacity)
+			.ToolTip(RowToolTip)
 		]
 	;
 }
 
 ////////////////////////////////////////////////////////////////////////////////////////////////////
 
-TSharedPtr<IToolTip> STableTreeViewCell::GetRowToolTip(const TSharedRef<ITableRow>& TableRow) const
-{
-	TSharedRef<STableTreeViewRow> Row = StaticCastSharedRef<STableTreeViewRow, ITableRow>(TableRow);
-	return Row->GetRowToolTip();
-}
-
-////////////////////////////////////////////////////////////////////////////////////////////////////
-
 FText STableTreeViewCell::GetValueAsText() const
 {
 	return ColumnPtr->GetValueAsText(*TableTreeNodePtr);
@@ -146,11 +136,11 @@ FText STableTreeViewCell::GetValueAsText() const
 
 ////////////////////////////////////////////////////////////////////////////////////////////////////
 
-TSharedRef<SWidget> STableTreeViewCell::GenerateWidgetForTableColumn(const FArguments& InArgs, const TSharedRef<ITableRow>& TableRow)
+TSharedRef<SWidget> STableTreeViewCell::GenerateWidgetForTableColumn(const FArguments& InArgs)
 {
 	TSharedRef<STextBlock> TextBox = SNew(STextBlock)
-		.TextStyle(FInsightsStyle::Get(), TEXT("TreeTable.Tooltip"))
-		.ColorAndOpacity(this, &STableTreeViewCell::GetStatsColorAndOpacity)
+		.TextStyle(FInsightsStyle::Get(), TEXT("TreeTable.NormalText"))
+		.ColorAndOpacity(this, &STableTreeViewCell::GetNormalTextColorAndOpacity)
 		.ShadowColorAndOpacity(this, &STableTreeViewCell::GetShadowColorAndOpacity);
 
 	if (ColumnPtr->IsDynamic())
@@ -163,15 +153,17 @@ TSharedRef<SWidget> STableTreeViewCell::GenerateWidgetForTableColumn(const FArgu
 		TextBox->SetText(CellText);
 	}
 
+	TSharedPtr<IToolTip> ColumnToolTip = ColumnPtr->GetValueFormatter()->GetCustomTooltip(*ColumnPtr, *TableTreeNodePtr);
+
 	return
 		SNew(SHorizontalBox)
-		.ToolTip(GetTooltip())
+		.ToolTip(ColumnToolTip)
 
 		+ SHorizontalBox::Slot()
 		.FillWidth(1.0f)
-		.VAlign(VAlign_Top)
 		.HAlign(ColumnPtr->GetHorizontalAlignment())
-		.Padding(FMargin(2.0f, 1.5f, 2.0f, 0.0f))
+		.VAlign(VAlign_Center)
+		.Padding(FMargin(2.0f, 0.0f, 2.0f, 0.0f))
 		[
 			TextBox
 		];
@@ -181,9 +173,109 @@ END_SLATE_FUNCTION_BUILD_OPTIMIZATION
 
 ////////////////////////////////////////////////////////////////////////////////////////////////////
 
-TSharedPtr<IToolTip> STableTreeViewCell::GetTooltip() const
+bool STableTreeViewCell::IsSelected() const
+{
+	return TableRow->IsItemSelected();
+}
+
+////////////////////////////////////////////////////////////////////////////////////////////////////
+
+FSlateColor STableTreeViewCell::GetIconColorAndOpacity() const
+{
+	bool bIsHoveredOrSelected = IsHovered() || IsSelected();
+
+	return bIsHoveredOrSelected ?
+		TableTreeNodePtr->GetColor() :
+		TableTreeNodePtr->GetColor().CopyWithNewOpacity(0.5f);
+}
+
+////////////////////////////////////////////////////////////////////////////////////////////////////
+
+FSlateColor STableTreeViewCell::GetDisplayNameColorAndOpacity() const
+{
+	bool bIsHoveredOrSelected = IsHovered() || IsSelected();
+
+	if (TableTreeNodePtr->IsFiltered())
+	{
+		return bIsHoveredOrSelected ?
+			TableTreeNodePtr->GetColor().CopyWithNewOpacity(0.5f) :
+			TableTreeNodePtr->GetColor().CopyWithNewOpacity(0.4f);
+	}
+	else
+	{
+		return bIsHoveredOrSelected ?
+			TableTreeNodePtr->GetColor() :
+			TableTreeNodePtr->GetColor().CopyWithNewOpacity(0.8f);
+	}
+}
+
+////////////////////////////////////////////////////////////////////////////////////////////////////
+
+FSlateColor STableTreeViewCell::GetExtraDisplayNameColorAndOpacity() const
+{
+	bool bIsHoveredOrSelected = IsHovered() || IsSelected();
+
+	if (TableTreeNodePtr->IsFiltered())
+	{
+		return bIsHoveredOrSelected ?
+			FLinearColor(0.3f, 0.3f, 0.3f, 0.5f) :
+			FLinearColor(0.3f, 0.3f, 0.3f, 0.4f);
+	}
+	else
+	{
+		return bIsHoveredOrSelected ?
+			FLinearColor(0.3f, 0.3f, 0.3f, 1.0f) :
+			FLinearColor(0.3f, 0.3f, 0.3f, 0.8f);
+	}
+}
+
+////////////////////////////////////////////////////////////////////////////////////////////////////
+
+FSlateColor STableTreeViewCell::GetNormalTextColorAndOpacity() const
+{
+	bool bIsHoveredOrSelected = IsHovered() || IsSelected();
+
+	if (TableTreeNodePtr->IsGroup())
+	{
+		if (TableTreeNodePtr->IsFiltered())
+		{
+			return bIsHoveredOrSelected ?
+				FStyleColors::ForegroundHover.GetSpecifiedColor().CopyWithNewOpacity(0.4f) :
+				FStyleColors::Foreground.GetSpecifiedColor().CopyWithNewOpacity(0.4f);
+		}
+		else
+		{
+			return bIsHoveredOrSelected ?
+				FStyleColors::ForegroundHover.GetSpecifiedColor().CopyWithNewOpacity(0.8f) :
+				FStyleColors::Foreground.GetSpecifiedColor().CopyWithNewOpacity(0.8f);
+		}
+	}
+	else
+	{
+		if (TableTreeNodePtr->IsFiltered())
+		{
+			return bIsHoveredOrSelected ?
+				FStyleColors::ForegroundHover.GetSpecifiedColor().CopyWithNewOpacity(0.5f) :
+				FStyleColors::Foreground.GetSpecifiedColor().CopyWithNewOpacity(0.5f);
+		}
+		else
+		{
+			return bIsHoveredOrSelected ?
+				FStyleColors::ForegroundHover :
+				FStyleColors::Foreground;
+		}
+	}
+}
+
+////////////////////////////////////////////////////////////////////////////////////////////////////
+
+FLinearColor STableTreeViewCell::GetShadowColorAndOpacity() const
 {
-	return ColumnPtr->GetValueFormatter()->GetCustomTooltip(*ColumnPtr, *TableTreeNodePtr);
+	const FLinearColor ShadowColor =
+		TableTreeNodePtr->IsFiltered() ?
+		FLinearColor(0.0f, 0.0f, 0.0f, 0.25f) :
+		FLinearColor(0.0f, 0.0f, 0.0f, 0.5f);
+	return ShadowColor;
 }
 
 ////////////////////////////////////////////////////////////////////////////////////////////////////
diff --git a/Engine/Plugins/Editor/AssetManagerEditor/Source/AssetManagerEditor/Private/Insights/Table/Widgets/STableTreeViewCell.h b/Engine/Plugins/Editor/AssetManagerEditor/Source/AssetManagerEditor/Private/Insights/Table/Widgets/STableTreeViewCell.h
--- a/Engine/Plugins/Editor/AssetManagerEditor/Source/AssetManagerEditor/Private/Insights/Table/Widgets/STableTreeViewCell.h
+++ b/Engine/Plugins/Editor/AssetManagerEditor/Source/AssetManagerEditor/Private/Insights/Table/Widgets/STableTreeViewCell.h
@@ -8,6 +8,7 @@
 
 // Insights
 #include "Insights/Table/ViewModels/TableTreeNode.h"
+#include "Insights/Table/Widgets/STableTreeViewRow.h"
 
 class ITableRow;
 class IToolTip;
@@ -34,12 +35,12 @@ public:
 		SLATE_ARGUMENT(bool, IsNameColumn)
 	SLATE_END_ARGS()
 
-	void Construct(const FArguments& InArgs, const TSharedRef<ITableRow>& TableRow);
+	void Construct(const FArguments& InArgs, const TSharedRef<ITableRow>& InTableRow);
 
 protected:
-	TSharedRef<SWidget> GenerateWidgetForColumn(const FArguments& InArgs, const TSharedRef<ITableRow>& TableRow);
-	TSharedRef<SWidget> GenerateWidgetForNameColumn(const FArguments& InArgs, const TSharedRef<ITableRow>& TableRow);
-	TSharedRef<SWidget> GenerateWidgetForTableColumn(const FArguments& InArgs, const TSharedRef<ITableRow>& TableRow);
+	TSharedRef<SWidget> GenerateWidgetForColumn(const FArguments& InArgs);
+	TSharedRef<SWidget> GenerateWidgetForNameColumn(const FArguments& InArgs);
+	TSharedRef<SWidget> GenerateWidgetForTableColumn(const FArguments& InArgs);
 
 	/**
 	 * The system will use this event to notify a widget that the cursor has entered it. This event is NOT bubbled.
@@ -103,7 +104,7 @@ protected:
 		return IsHovered() ? EVisibility::Visible : EVisibility::Hidden;
 	}
 
-	TSharedPtr<class IToolTip> GetRowToolTip(const TSharedRef<ITableRow>& TableRow) const;
+	bool IsSelected() const;
 
 	const FSlateBrush* GetIcon() const
 	{
@@ -127,56 +128,15 @@ protected:
 
 	FText GetValueAsText() const;
 
-	FSlateColor GetIconColorAndOpacity() const
-	{
-		const FLinearColor TextColor =
-			IsHovered() ?
-				TableTreeNodePtr->GetColor() :
-				TableTreeNodePtr->GetColor().CopyWithNewOpacity(0.5f);
-		return TextColor;
-	}
-
-	FSlateColor GetColorAndOpacity() const
-	{
-		const FLinearColor TextColor =
-			TableTreeNodePtr->IsFiltered() ?
-				TableTreeNodePtr->GetColor().CopyWithNewOpacity(0.5f) :
-				TableTreeNodePtr->GetColor();
-		return TextColor;
-	}
-
-	FSlateColor GetExtraColorAndOpacity() const
-	{
-		const FLinearColor TextColor =
-			TableTreeNodePtr->IsFiltered() ?
-				FLinearColor(0.5f, 0.5f, 0.5f, 0.5f) :
-				FLinearColor(0.5f, 0.5f, 0.5f, 1.0f);
-		return TextColor;
-	}
-
-	FSlateColor GetStatsColorAndOpacity() const
-	{
-		const FLinearColor TextColor =
-			TableTreeNodePtr->IsGroup() ?
-				FLinearColor(0.7f, 0.7f, 0.7f, 1.0f) :
-			TableTreeNodePtr->IsFiltered() ?
-				FLinearColor(1.0f, 1.0f, 1.0f, 0.5f) :
-				FLinearColor(1.0f, 1.0f, 1.0f, 1.0f);
-		return TextColor;
-	}
-
-	FLinearColor GetShadowColorAndOpacity() const
-	{
-		const FLinearColor ShadowColor =
-			TableTreeNodePtr->IsFiltered() ?
-				FLinearColor(0.0f, 0.0f, 0.0f, 0.25f) :
-				FLinearColor(0.0f, 0.0f, 0.0f, 0.5f);
-		return ShadowColor;
-	}
-
-	TSharedPtr<IToolTip> GetTooltip() const;
+	FSlateColor GetIconColorAndOpacity() const;
+	FSlateColor GetDisplayNameColorAndOpacity() const;
+	FSlateColor GetExtraDisplayNameColorAndOpacity() const;
+	FSlateColor GetNormalTextColorAndOpacity() const;
+	FLinearColor GetShadowColorAndOpacity() const;
 
 protected:
+	TSharedPtr<ITableRow> TableRow;
+
 	/** A shared pointer to the table view model. */
 	TSharedPtr<FTable> TablePtr;
 
diff --git a/Engine/Plugins/Editor/AssetManagerEditor/Source/AssetManagerEditor/Private/SAssetAuditBrowser.cpp b/Engine/Plugins/Editor/AssetManagerEditor/Source/AssetManagerEditor/Private/SAssetAuditBrowser.cpp
--- a/Engine/Plugins/Editor/AssetManagerEditor/Source/AssetManagerEditor/Private/SAssetAuditBrowser.cpp
+++ b/Engine/Plugins/Editor/AssetManagerEditor/Source/AssetManagerEditor/Private/SAssetAuditBrowser.cpp
@@ -25,6 +25,7 @@
 #include "Editor.h"
 
 #if ASSET_TABLE_TREE_VIEW_ENABLED
+#include "Styling/StyleColors.h"
 #include "TreeView/AssetTable.h"
 #include "TreeView/SAssetTableTreeView.h"
 #endif
@@ -743,7 +744,7 @@ void SAssetAuditBrowser::PopulateAssetTableRow(const FAssetData& AssetData, clas
 	OutRow->NativeClass = AssetData.AssetClassPath.GetAssetName().ToString();
 	OutRow->Color = USlateThemeManager::Get().GetColor((EStyleColor)((int32)EStyleColor::AccentBlue + FMath::Rand() % 8)); // TODO: color by type!?
 }
-#endif
+#endif // ASSET_TABLE_TREE_VIEW_ENABLED
 
 void SAssetAuditBrowser::RefreshAssetView()
 {
diff --git a/Engine/Plugins/Editor/AssetManagerEditor/Source/AssetManagerEditor/Private/TreeView/AssetDependencyGrouping.cpp b/Engine/Plugins/Editor/AssetManagerEditor/Source/AssetManagerEditor/Private/TreeView/AssetDependencyGrouping.cpp
--- a/Engine/Plugins/Editor/AssetManagerEditor/Source/AssetManagerEditor/Private/TreeView/AssetDependencyGrouping.cpp
+++ b/Engine/Plugins/Editor/AssetManagerEditor/Source/AssetManagerEditor/Private/TreeView/AssetDependencyGrouping.cpp
@@ -31,6 +31,10 @@ void FAssetDependencyGrouping::GroupNodes(const TArray<UE::Insights::FTableTreeN
 	TSharedPtr<FAssetTable> AssetTable = StaticCastSharedPtr<FAssetTable>(InParentTable.Pin());
 	check(AssetTable.IsValid());
 
+	const FSlateBrush* IconBrush = FBaseTreeNode::GetDefaultIcon(true);
+	FLinearColor AssetGroupNodeColor(1.0f, 1.0f, 1.0f, 1.0f);
+	FLinearColor DependenciesGroupNodeColor(0.75f, 0.5f, 1.0f, 1.0f);
+
 	for (FTableTreeNodePtr NodePtr : Nodes)
 	{
 		if (InAsyncOperationProgress.ShouldCancelAsyncOp())
@@ -61,7 +65,7 @@ void FAssetDependencyGrouping::GroupNodes(const TArray<UE::Insights::FTableTreeN
 		FAssetTreeNode& AssetNode = NodePtr->As<FAssetTreeNode>();
 		const FAssetTableRow& Asset = AssetNode.GetAssetChecked();
 
-		FTableTreeNodePtr GroupPtr = MakeShared<FTableTreeNode>(FName(Asset.GetName(), 0), InParentTable);
+		FTableTreeNodePtr GroupPtr = MakeShared<FCustomTableTreeNode>(FName(Asset.GetName(), 0), InParentTable, IconBrush, AssetGroupNodeColor);
 		GroupPtr->SetExpansion(false);
 		ParentGroup.AddChildAndSetGroupPtr(GroupPtr);
 
@@ -71,7 +75,7 @@ void FAssetDependencyGrouping::GroupNodes(const TArray<UE::Insights::FTableTreeN
 		const TArray<int32>& Dependencies = Asset.GetDependencies();
 		if (Dependencies.Num() > 0)
 		{
-			DepGroupPtr = MakeShared<FTableTreeNode>(FName(TEXT("Dependencies"), 0), InParentTable);
+			DepGroupPtr = MakeShared<FCustomTableTreeNode>(FName(TEXT("Dependencies"), 0), InParentTable, IconBrush, DependenciesGroupNodeColor);
 			DepGroupPtr->SetExpansion(false);
 			GroupPtr->AddChildAndSetGroupPtr(DepGroupPtr);
 		}
diff --git a/Engine/Plugins/Editor/AssetManagerEditor/Source/AssetManagerEditor/Private/TreeView/AssetTreeNode.cpp b/Engine/Plugins/Editor/AssetManagerEditor/Source/AssetManagerEditor/Private/TreeView/AssetTreeNode.cpp
--- a/Engine/Plugins/Editor/AssetManagerEditor/Source/AssetManagerEditor/Private/TreeView/AssetTreeNode.cpp
+++ b/Engine/Plugins/Editor/AssetManagerEditor/Source/AssetManagerEditor/Private/TreeView/AssetTreeNode.cpp
@@ -2,4 +2,30 @@
 
 #include "AssetTreeNode.h"
 
+#include "Insights/Common/InsightsStyle.h"
+
 INSIGHTS_IMPLEMENT_RTTI(FAssetTreeNode)
+
+const FSlateBrush* FAssetTreeNode::GetIcon() const
+{
+#if 1 // debug / mock code
+	if (RowId.RowIndex % 7)
+	{
+		return UE::Insights::FInsightsStyle::GetBrush("Icons.Package.TreeItem");
+	}
+	if (RowId.RowIndex % 3)
+	{
+		return UE::Insights::FInsightsStyle::GetBrush("Icons.Asset.TreeItem");
+	}
+#endif
+
+	//const FAssetTableRow& AssetRpw = GetAssetChecked();
+	//TODO: switch (AssetRow->GetType()) --> icon
+
+	return FBaseTreeNode::GetDefaultIcon(false); // default icon for leaf nodes
+}
+
+FLinearColor FAssetTreeNode::GetColor() const
+{
+	return GetAssetChecked().GetColor();
+}
diff --git a/Engine/Plugins/Editor/AssetManagerEditor/Source/AssetManagerEditor/Private/TreeView/AssetTreeNode.h b/Engine/Plugins/Editor/AssetManagerEditor/Source/AssetManagerEditor/Private/TreeView/AssetTreeNode.h
--- a/Engine/Plugins/Editor/AssetManagerEditor/Source/AssetManagerEditor/Private/TreeView/AssetTreeNode.h
+++ b/Engine/Plugins/Editor/AssetManagerEditor/Source/AssetManagerEditor/Private/TreeView/AssetTreeNode.h
@@ -29,10 +29,8 @@ public:
 	FAssetTable& GetAssetTableChecked() const { return *AssetTable; }
 	const FAssetTableRow& GetAssetChecked() const { return AssetTable->GetAssetChecked(RowId.RowIndex); }
 
-	virtual FLinearColor GetColor() const override
-	{
-		return GetAssetChecked().GetColor();
-	}
+	virtual const FSlateBrush* GetIcon() const override;
+	virtual FLinearColor GetColor() const override;
 
 private:
 	FAssetTable* AssetTable;
diff --git a/Engine/Plugins/Editor/AssetManagerEditor/Source/AssetManagerEditor/Private/TreeView/SAssetTableTreeView.cpp b/Engine/Plugins/Editor/AssetManagerEditor/Source/AssetManagerEditor/Private/TreeView/SAssetTableTreeView.cpp
--- a/Engine/Plugins/Editor/AssetManagerEditor/Source/AssetManagerEditor/Private/TreeView/SAssetTableTreeView.cpp
+++ b/Engine/Plugins/Editor/AssetManagerEditor/Source/AssetManagerEditor/Private/TreeView/SAssetTableTreeView.cpp
@@ -114,10 +114,6 @@ void SAssetTableTreeView::RebuildTree(bool bResync)
 			//FName NodeName(BaseNodeName, AssetIndex + 1);
 			FName NodeName(Asset->GetName());
 			FAssetTreeNodePtr NodePtr = MakeShared<FAssetTreeNode>(NodeName, AssetTable, AssetIndex);
-			if (AssetIndex % 3) // TODO
-				NodePtr->SetIcon(UE::Insights::FInsightsStyle::GetBrush("Icons.Asset.TreeItem"));
-			else if (AssetIndex % 7) // TODO
-				NodePtr->SetIcon(UE::Insights::FInsightsStyle::GetBrush("Icons.Package.TreeItem"));
 			TableTreeNodes.Add(NodePtr);
 		}
 		ensure(TableTreeNodes.Num() == VisibleAssetCount);
