AssetManagerEditor: Asset Tree View improvements. Always enabled log message about tree view rebuilt.

diff --git a/Engine/Plugins/Editor/AssetManagerEditor/Source/AssetManagerEditor/Private/TreeView/SAssetTableTreeView.cpp b/Engine/Plugins/Editor/AssetManagerEditor/Source/AssetManagerEditor/Private/TreeView/SAssetTableTreeView.cpp
--- a/Engine/Plugins/Editor/AssetManagerEditor/Source/AssetManagerEditor/Private/TreeView/SAssetTableTreeView.cpp
+++ b/Engine/Plugins/Editor/AssetManagerEditor/Source/AssetManagerEditor/Private/TreeView/SAssetTableTreeView.cpp
@@ -151,12 +151,9 @@ void SAssetTableTreeView::RebuildTree(bool bResync)
 
 	Stopwatch.Stop();
 	const double TotalTime = Stopwatch.GetAccumulatedTime();
-	if (TotalTime > 0.01)
-	{
-		const double SyncTime = SyncStopwatch.GetAccumulatedTime();
-		UE_LOG(LogInsights, Log, TEXT("[AssetTree] Tree view rebuilt in %.4fs (sync: %.4fs + update: %.4fs) --> %d nodes (%d added)"),
-			TotalTime, SyncTime, TotalTime - SyncTime, TableTreeNodes.Num(), TableTreeNodes.Num() - PreviousNodeCount);
-	}
+	const double SyncTime = SyncStopwatch.GetAccumulatedTime();
+	UE_LOG(LogInsights, Log, TEXT("[AssetTree] Tree view rebuilt in %.4fs (sync: %.4fs + update: %.4fs) --> %d nodes (%d added)"),
+		TotalTime, SyncTime, TotalTime - SyncTime, TableTreeNodes.Num(), TableTreeNodes.Num() - PreviousNodeCount);
 }
 
 ////////////////////////////////////////////////////////////////////////////////////////////////////
