AssetManagerEditor: Asset Tree View improvements.
  - Fixed crash when clearing assets.
  - Added horizontal splitter between the legacy asset list view and the asset tree view (to be able to resize the height of these views).

diff --git a/Engine/Plugins/Editor/AssetManagerEditor/Source/AssetManagerEditor/Private/SAssetAuditBrowser.cpp b/Engine/Plugins/Editor/AssetManagerEditor/Source/AssetManagerEditor/Private/SAssetAuditBrowser.cpp
--- a/Engine/Plugins/Editor/AssetManagerEditor/Source/AssetManagerEditor/Private/SAssetAuditBrowser.cpp
+++ b/Engine/Plugins/Editor/AssetManagerEditor/Source/AssetManagerEditor/Private/SAssetAuditBrowser.cpp
@@ -366,6 +366,14 @@ void SAssetAuditBrowser::Construct(const FArguments& InArgs)
 
 	this->ChildSlot
 	[
+#if ASSET_TABLE_TREE_VIEW_ENABLED
+	SNew(SSplitter)
+	.Orientation(EOrientation::Orient_Vertical)
+	+ SSplitter::Slot()
+	.Value(0.5f)
+	.MinSize(100.0f)
+	[
+#endif // ASSET_TABLE_TREE_VIEW_ENABLED
 		SNew(SVerticalBox)
 		+SVerticalBox::Slot()
 		.AutoHeight()
@@ -518,16 +526,11 @@ void SAssetAuditBrowser::Construct(const FArguments& InArgs)
 			]
 		]
 #if ASSET_TABLE_TREE_VIEW_ENABLED
-		+SVerticalBox::Slot()
-		.FillHeight(1.f)
-		[
-			SNew(SBorder)
-			.Padding(FMargin(3))
-			.BorderImage(FAppStyle::GetBrush("ToolPanel.GroupBorder"))
-			[
-				SAssignNew(AssetTableTreeView, SAssetTableTreeView, AssetTable)
-			]
-		]
+	]
+	+ SSplitter::Slot()
+	[
+		SAssignNew(AssetTableTreeView, SAssetTableTreeView, AssetTable)
+	]
 #endif //ASSET_TABLE_TREE_VIEW_ENABLED
 	];
 
@@ -773,6 +776,11 @@ void SAssetAuditBrowser::RefreshAssetView()
 		TSharedPtr<FAssetTable> AssetTable = AssetTableTreeView->GetAssetTable();
 		if (AssetTable.IsValid())
 		{
+			// Clears all tree nodes (that references the previous assets).
+			AssetTable->SetVisibleAssetCount(0);
+			AssetTableTreeView->RebuildTree(true);
+
+			// Now is safe to clear the previous assets.
 			AssetTable->ClearAllData();
 
 			TMap<FAssetData, int32> AssetToIndexMap;
