AssetManagerEditor: Improvements to Asset Tree view.
  - Lazy create asset dependencies nodes (when grouping by dependency). Moved main asset into a _Self_ group when there are Dependencies and directly in hierarchy when it has no dependencies.
  - Coloring of assets by type.
  - Added logs for asset table rebuild time.
  - Refactoring on FBaseTreeNode, FTableTreeNode and STableTreeView code.

diff --git a/Engine/Plugins/Editor/AssetManagerEditor/Source/AssetManagerEditor/Private/Insights/Table/ViewModels/BaseTreeNode.h b/Engine/Plugins/Editor/AssetManagerEditor/Source/AssetManagerEditor/Private/Insights/Table/ViewModels/BaseTreeNode.h
--- a/Engine/Plugins/Editor/AssetManagerEditor/Source/AssetManagerEditor/Private/Insights/Table/ViewModels/BaseTreeNode.h
+++ b/Engine/Plugins/Editor/AssetManagerEditor/Source/AssetManagerEditor/Private/Insights/Table/ViewModels/BaseTreeNode.h
@@ -4,6 +4,7 @@
 
 #include "CoreMinimal.h"
 #include "Math/Color.h"
+#include "Templates/Function.h"
 
 #include "Insights/Common/SimpleRtti.h"
 
@@ -85,7 +86,7 @@ public:
 	}
 
 	/**
-	 * @return a name of this node.
+	 * @returns a name of this node.
 	 */
 	const FName& GetName() const
 	{
@@ -93,19 +94,22 @@ public:
 	}
 
 	/**
-	 * @return a name of this node to display in tree view.
+	 * @returns a name of this node to display in tree view.
 	 */
 	virtual const FText GetDisplayName() const;
 
 	/**
-	 * @return a name suffix for this node to display in tree view. Ex.: group nodes may include additional info, like the number of visible / total children.
+	 * @returns a name suffix for this node to display in tree view. Ex.: group nodes may include additional info, like the number of visible / total children.
 	 */
 	virtual const FText GetExtraDisplayName() const;
 
+	/**
+	 * @returns true if the node has an extra display name (a name suffix to display in tree view).
+	 */
 	virtual bool HasExtraDisplayName() const;
 
 	/**
-	 * @return a descriptive text for this node to display in a tooltip.
+	 * @returns a descriptive text for this node to display in a tooltip.
 	 */
 	virtual const FText GetTooltip() const
 	{
@@ -124,12 +128,12 @@ public:
 	}
 
 	/**
-	 * @return the default icon for a group/leaf node.
+	 * @returns the default icon for a group/leaf node.
 	 */
 	static const FSlateBrush* GetDefaultIcon(bool bIsGroupNode);
 
 	/**
-	 * @return a brush icon for this node.
+	 * @returns a brush icon for this node.
 	 */
 	virtual const FSlateBrush* GetIcon() const
 	{
@@ -137,12 +141,12 @@ public:
 	}
 
 	/**
-	 * @return the default color tint for a group/leaf node.
+	 * @returns the default color tint for a group/leaf node.
 	 */
 	static FLinearColor GetDefaultColor(bool bIsGroupNode);
 
 	/**
-	 * @return the color tint to be used for the icon and the name text of this node.
+	 * @returns the color tint to be used for the icon and the name text of this node.
 	 */
 	virtual FLinearColor GetColor() const
 	{
@@ -150,7 +154,7 @@ public:
 	}
 
 	/**
-	 * @return a pointer to a data context for this node.
+	 * @returns a pointer to a data context for this node. Only available for group nodes.
 	 */
 	virtual void* GetContext() const
 	{
@@ -158,7 +162,7 @@ public:
 	}
 
 	/**
-	 * Sets a pointer to a data context for this node.
+	 * Sets a pointer to a data context for this node. Only available for group nodes.
 	 */
 	virtual void SetContext(void* InContext)
 	{
@@ -169,7 +173,7 @@ public:
 	}
 
 	/**
-	 * @return true, if this node is a group node.
+	 * @returns true if this node is a group node.
 	 */
 	bool IsGroup() const
 	{
@@ -188,23 +192,67 @@ public:
 	}
 
 	/**
-	 * @return a const reference to the child nodes of this group.
+	 * @returns the number of children nodes.
+	 */
+	int32 GetChildrenCount() const
+	{
+		return GroupData->Children.Num();
+	}
+
+	/**
+	 * @returns the number of filtered children nodes.
+	 */
+	int32 GetFilteredChildrenCount() const
+	{
+		return GroupData->FilteredChildren.Num();
+	}
+
+	/**
+	 * Enumerates the children nodes.
+	 */
+	void EnumerateChildren(TFunction<bool(const FBaseTreeNodePtr&)> Callback) const
+	{
+		for (const FBaseTreeNodePtr& ChildNode : GroupData->Children)
+		{
+			if (!Callback(ChildNode))
+			{
+				break;
+			}
+		}
+	}
+
+	/**
+	 * Enumerates the filtered children nodes.
+	 */
+	void EnumerateFilteredChildren(TFunction<bool(const FBaseTreeNodePtr&)> Callback) const
+	{
+		for (const FBaseTreeNodePtr& ChildNode : GroupData->FilteredChildren)
+		{
+			if (!Callback(ChildNode))
+			{
+				break;
+			}
+		}
+	}
+
+	/**
+	 * @returns a const reference to the children nodes of this group.
 	 */
-	FORCEINLINE_DEBUGGABLE const TArray<FBaseTreeNodePtr>& GetChildren() const
+	const TArray<FBaseTreeNodePtr>& GetChildren() const
 	{
 		return GroupData->Children;
 	}
 
 	/**
-	 * @return a const reference to the child nodes that should be visible to the UI based on filtering.
+	 * @returns a const reference to the children nodes that should be visible to the UI based on filtering.
 	 */
-	FORCEINLINE_DEBUGGABLE const TArray<FBaseTreeNodePtr>& GetFilteredChildren() const
+	const TArray<FBaseTreeNodePtr>& GetFilteredChildren() const
 	{
 		return GroupData->FilteredChildren;
 	}
 
 	/**
-	 * @return a weak reference to the group of this node, may be invalid.
+	 * @returns a weak reference to the group of this node, may be invalid.
 	 */
 	FBaseTreeNodeWeak GetGroupPtr() const
 	{
@@ -212,16 +260,11 @@ public:
 	}
 
 	/**
-	 * @return a name of the group node that this node belongs to.
+	 * @returns a weak reference to the group of this node, may be invalid.
 	 */
-	const FName& GetGroupName() const
+	FBaseTreeNodePtr GetParentNode() const
 	{
-		return GroupPtr.Pin()->GetName();
-	}
-
-	virtual bool IsFiltered() const
-	{
-		return false;
+		return GroupPtr.Pin();
 	}
 
 	void SortChildrenAscending(const ITableCellValueSorter& Sorter);
@@ -233,16 +276,16 @@ public:
 		GroupData->Children.Sort(Predicate);
 	}
 
-	/** Adds specified child to the children and sets group for it. */
-	FORCEINLINE_DEBUGGABLE void AddChildAndSetGroupPtr(const FBaseTreeNodePtr& ChildPtr)
+	/** Adds specified node to the children nodes. Also sets the current node as the group of the specified node. */
+	void AddChildAndSetGroupPtr(const FBaseTreeNodePtr& ChildPtr)
 	{
-		ChildPtr->GroupPtr = AsShared();
+		ChildPtr->GroupPtr = AsWeak();
 		GroupData->Children.Add(ChildPtr);
 	}
 
 	void SetGroupPtrForAllChildren()
 	{
-		FBaseTreeNodeWeak ThisNode = AsShared();
+		FBaseTreeNodeWeak ThisNode = AsWeak();
 		for (FBaseTreeNodePtr& NodePtr : GroupData->Children)
 		{
 			NodePtr->GroupPtr = ThisNode;
@@ -258,40 +301,65 @@ public:
 	}
 
 	/** Clears children. */
-	FORCEINLINE_DEBUGGABLE void ClearChildren(int32 NewSize = 0)
+	void ClearChildren(int32 NewSize = 0)
 	{
 		ResetGroupPtrForAllChildren();
 		GroupData->Children.Reset(NewSize);
 	}
 
-	FORCEINLINE_DEBUGGABLE void SwapChildren(TArray<FBaseTreeNodePtr>& NewChildren)
+	void SwapChildren(TArray<FBaseTreeNodePtr>& NewChildren)
 	{
 		ResetGroupPtrForAllChildren();
 		Swap(NewChildren, GroupData->Children);
 		SetGroupPtrForAllChildren();
 	}
 
-	FORCEINLINE_DEBUGGABLE void SwapChildrenFast(TArray<FBaseTreeNodePtr>& NewChildren)
+	void SwapChildrenFast(TArray<FBaseTreeNodePtr>& NewChildren)
 	{
 		Swap(NewChildren, GroupData->Children);
 	}
 
-	/** Adds specified child to the filtered children. */
-	FORCEINLINE_DEBUGGABLE void AddFilteredChild(const FBaseTreeNodePtr& ChildPtr)
+	/** Adds specified child to the filtered children nodes. */
+	void AddFilteredChild(const FBaseTreeNodePtr& ChildPtr)
 	{
 		GroupData->FilteredChildren.Add(ChildPtr);
 	}
 
-	/** Clears filtered children. */
+	/** Clears the filtered children nodes. */
 	void ClearFilteredChildren(int32 NewSize = 0)
 	{
 		GroupData->FilteredChildren.Reset(NewSize);
 	}
 
+	/**
+	 * Resets the filtered children nodes.
+	 * The filtered array of children nodes will be initialized with the unfiltered array of children.
+	 */
+	void ResetFilteredChildren()
+	{
+		const int32 Size = GroupData->Children.Num();
+		GroupData->FilteredChildren.Reset();
+		GroupData->FilteredChildren.Append(GroupData->Children);
+	}
+
+	/**
+	 * Resets the filtered children for this node and also recursively for all children nodes.
+	 * The filtered array of children nodes will be initialized with the unfiltered array of children.
+	 */
+	void ResetFilteredChildrenRec()
+	{
+		this->ResetFilteredChildren();
+		for (FBaseTreeNodePtr Node : GroupData->Children)
+		{
+			Node->ResetFilteredChildrenRec();
+		}
+	}
+
 	bool IsExpanded() const
 	{
 		return GroupData->bIsExpanded;
 	}
+
 	void SetExpansion(bool bOnOff)
 	{
 		if (IsGroup())
@@ -302,9 +370,9 @@ public:
 
 protected:
 	/**
-	 * @return a reference to the child nodes of this group.
+	 * @returns a reference to the children nodes of this group.
 	 */
-	FORCEINLINE_DEBUGGABLE TArray<FBaseTreeNodePtr>& GetChildrenMutable()
+	TArray<FBaseTreeNodePtr>& GetChildrenMutable()
 	{
 		return GroupData->Children;
 	}
diff --git a/Engine/Plugins/Editor/AssetManagerEditor/Source/AssetManagerEditor/Private/Insights/Table/ViewModels/TableTreeNode.h b/Engine/Plugins/Editor/AssetManagerEditor/Source/AssetManagerEditor/Private/Insights/Table/ViewModels/TableTreeNode.h
--- a/Engine/Plugins/Editor/AssetManagerEditor/Source/AssetManagerEditor/Private/Insights/Table/ViewModels/TableTreeNode.h
+++ b/Engine/Plugins/Editor/AssetManagerEditor/Source/AssetManagerEditor/Private/Insights/Table/ViewModels/TableTreeNode.h
@@ -14,6 +14,7 @@ namespace Insights
 {
 
 class FTable;
+class STableTreeView;
 
 struct FTableRowId
 {
@@ -79,9 +80,11 @@ public:
 	FTableRowId GetRowId() const { return RowId; }
 	int32 GetRowIndex() const { return RowId.RowIndex; }
 
+	//////////////////////////////////////////////////
+	// Aggregation
+
 	void InitAggregatedValues()
 	{
-		check(IsGroup());
 		if (!AggregatedValues)
 		{
 			AggregatedValues = new TMap<FName, FTableCellValue>();
@@ -97,17 +100,53 @@ public:
 		}
 	}
 
-	void ResetAggregatedValues() { CleanupAggregatedValues(); }
-	void ResetAggregatedValues(const FName& ColumnId) { if (AggregatedValues) { AggregatedValues->Remove(ColumnId); } }
-	bool HasAggregatedValue(const FName& ColumnId) const { return AggregatedValues && AggregatedValues->Contains(ColumnId); }
-	const FTableCellValue* FindAggregatedValue(const FName& ColumnId) const { return AggregatedValues ? AggregatedValues->Find(ColumnId) : nullptr; }
-	const FTableCellValue& GetAggregatedValue(const FName& ColumnId) const { return AggregatedValues->FindChecked(ColumnId); }
-	void AddAggregatedValue(const FName& ColumnId, const FTableCellValue& Value) { InitAggregatedValues(); AggregatedValues->Add(ColumnId, Value); }
-	void SetAggregatedValue(const FName& ColumnId, const FTableCellValue& Value) { InitAggregatedValues(); (*AggregatedValues)[ColumnId] = Value; }
+	void ResetAggregatedValues()
+	{
+		CleanupAggregatedValues();
+	}
+
+	void ResetAggregatedValue(const FName& ColumnId)
+	{
+		if (AggregatedValues)
+		{
+			AggregatedValues->Remove(ColumnId);
+		}
+	}
+
+	bool HasAggregatedValue(const FName& ColumnId) const
+	{
+		return AggregatedValues && AggregatedValues->Contains(ColumnId);
+	}
 
-	virtual bool IsFiltered() const override { return bIsFiltered; }
+	const FTableCellValue* FindAggregatedValue(const FName& ColumnId) const
+	{
+		return AggregatedValues ? AggregatedValues->Find(ColumnId) : nullptr;
+	}
+
+	const FTableCellValue& GetAggregatedValue(const FName& ColumnId) const
+	{
+		return AggregatedValues->FindChecked(ColumnId);
+	}
+
+	void AddAggregatedValue(const FName& ColumnId, const FTableCellValue& Value)
+	{
+		InitAggregatedValues();
+		AggregatedValues->Add(ColumnId, Value);
+	}
+
+	void SetAggregatedValue(const FName& ColumnId, const FTableCellValue& Value)
+	{
+		InitAggregatedValues();
+		(*AggregatedValues)[ColumnId] = Value;
+	}
+
+	//////////////////////////////////////////////////
+
+	virtual bool IsFiltered() const { return bIsFiltered; }
 	virtual void SetIsFiltered(bool InValue) { bIsFiltered = InValue; }
 
+	virtual bool OnLazyCreateChildren(TSharedPtr<STableTreeView> InTableTreeView) { return false; }
+
 protected:
 	TWeakPtr<FTable> ParentTable;
 	FTableRowId RowId;
diff --git a/Engine/Plugins/Editor/AssetManagerEditor/Source/AssetManagerEditor/Private/Insights/Table/Widgets/STableTreeView.cpp b/Engine/Plugins/Editor/AssetManagerEditor/Source/AssetManagerEditor/Private/Insights/Table/Widgets/STableTreeView.cpp
--- a/Engine/Plugins/Editor/AssetManagerEditor/Source/AssetManagerEditor/Private/Insights/Table/Widgets/STableTreeView.cpp
+++ b/Engine/Plugins/Editor/AssetManagerEditor/Source/AssetManagerEditor/Private/Insights/Table/Widgets/STableTreeView.cpp
@@ -278,6 +278,7 @@ void STableTreeView::ConstructWidget(TSharedPtr<FTable> InTablePtr)
 				.OnGetChildren(this, &STableTreeView::TreeView_OnGetChildren)
 				.OnGenerateRow(this, &STableTreeView::TreeView_OnGenerateRow)
 				.OnSelectionChanged(this, &STableTreeView::TreeView_OnSelectionChanged)
+				.OnExpansionChanged(this, &STableTreeView::TreeView_OnExpansionChanged)
 				.OnMouseButtonDoubleClick(this, &STableTreeView::TreeView_OnMouseButtonDoubleClick)
 				.OnContextMenuOpening(FOnContextMenuOpening::CreateSP(this, &STableTreeView::TreeView_GetMenuContent))
 				.ItemHeight(16.0f)
@@ -302,13 +303,13 @@ void STableTreeView::ConstructWidget(TSharedPtr<FTable> InTablePtr)
 			.Padding(FMargin(0.0f, 30.0f, 0.0f, 0.0f))
 			[
 				SNew(SBorder)
-				.BorderImage(FAppStyle::Get().GetBrush("PopupText.Background"))
-				[
-					SNew(SHorizontalBox)
-					.Visibility_Lambda([this]() -> EVisibility
+				.BorderImage(FAppStyle::Get().GetBrush("Brushes.Panel"))
+				.Visibility_Lambda([this]() -> EVisibility
 					{
 						return TreeViewBannerText.IsEmpty() ? EVisibility::Collapsed : EVisibility::Visible;
 					})
+				[
+					SNew(SHorizontalBox)
 
 					+ SHorizontalBox::Slot()
 					.VAlign(VAlign_Center)
@@ -967,7 +968,7 @@ void STableTreeView::Tick(const FGeometry& AllottedGeometry, const double InCurr
 		{
 			// We wait for the TreeView to be refreshed before dispatching the tasks.
 			// This should make the TreeView release all of it's shared pointers to nodes to prevent
-			// the TreeView (MainThread) and the tasks from accesing the non-thread safe shared pointers at the same time.
+			// the TreeView (MainThread) and the tasks from accessing the non-thread safe shared pointers at the same time.
 			if (!TreeView->IsPendingRefresh())
 			{
 				DispatchEvent->DispatchSubsequents();
@@ -1040,6 +1041,13 @@ void STableTreeView::OnFilteringChanged()
 ////////////////////////////////////////////////////////////////////////////////////////////////////
 
 void STableTreeView::ApplyFiltering()
+{
+	ApplyFiltering(Root);
+}
+
+////////////////////////////////////////////////////////////////////////////////////////////////////
+
+void STableTreeView::ApplyFiltering(FTableTreeNodePtr SubTreeRootPtr)
 {
 	FStopwatch Stopwatch;
 	Stopwatch.Start();
@@ -1047,12 +1055,12 @@ void STableTreeView::ApplyFiltering()
 	// Apply filter to all groups and its children.
 	if (FilterConfigurator_HasFilters())
 	{
-		ApplyAdvancedFiltersForNode(Root);
+		ApplyAdvancedFiltersForNode(SubTreeRootPtr);
 	}
 	else
 	{
 		const bool bFilterIsEmpty = bRunInAsyncMode ? CurrentAsyncOpTextFilter->GetRawFilterText().IsEmpty() : TextFilter->GetRawFilterText().IsEmpty();
-		ApplyHierarchicalFilterForNode(Root, bFilterIsEmpty);
+		ApplyHierarchicalFilterForNode(SubTreeRootPtr, bFilterIsEmpty);
 	}
 
 	FilteredGroupNodes.Reset();
@@ -1085,7 +1093,7 @@ void STableTreeView::ApplyFiltering()
 		UE_LOG(LogInsights, Log, TEXT("[Tree - %s] Filtering completed in %.3fs."), *Table->GetDisplayName().ToString(), FilteringTime);
 	}
 
-	UpdateAggregatedValues(*Root);
+	UpdateAggregatedValuesRec(*Root);
 
 	// Cannot call TreeView functions from other threads than MainThread and SlateThread.
 	if (!bRunInAsyncMode)
@@ -1307,8 +1315,28 @@ void STableTreeView::TreeView_OnSelectionChanged(FTableTreeNodePtr SelectedItem,
 
 ////////////////////////////////////////////////////////////////////////////////////////////////////
 
+void STableTreeView::TreeView_OnExpansionChanged(FTableTreeNodePtr TreeNode, bool bShouldBeExpanded)
+{
+}
+
+////////////////////////////////////////////////////////////////////////////////////////////////////
+
 void STableTreeView::TreeView_OnGetChildren(FTableTreeNodePtr InParent, TArray<FTableTreeNodePtr>& OutChildren)
 {
+	if (InParent->OnLazyCreateChildren(SharedThis(this)))
+	{
+		// The filtering is ignored for lazy expanded nodes.
+		InParent->ResetFilteredChildrenRec();
+
+		UpdateAggregatedValuesRec(*InParent);
+		FTableTreeNodePtr Node = StaticCastSharedPtr<FTableTreeNode>(InParent->GetParentNode());
+		while (Node)
+		{
+			UpdateAggregatedValuesSingleNode(*Node);
+			Node = StaticCastSharedPtr<FTableTreeNode>(Node->GetParentNode());
+		}
+	}
+
 	const TArray<FBaseTreeNodePtr>& FilteredChildren = InParent->GetFilteredChildren();
 	for (const FBaseTreeNodePtr& NodePtr : FilteredChildren)
 	{
@@ -1495,7 +1523,7 @@ void STableTreeView::CreateGroups(const TArray<TSharedPtr<FTreeNodeGrouping>>& G
 		UE_LOG(LogInsights, Log, TEXT("[Tree - %s] Grouping completed in %.3fs."), *Table->GetDisplayName().ToString(), GroupingTime);
 	}
 
-	UpdateAggregatedValues(*Root);
+	UpdateAggregatedValuesRec(*Root);
 }
 
 ////////////////////////////////////////////////////////////////////////////////////////////////////
@@ -1575,12 +1603,152 @@ void STableTreeView::GroupNodesRec(const TArray<FTableTreeNodePtr>& Nodes, FTabl
 
 ////////////////////////////////////////////////////////////////////////////////////////////////////
 
-void STableTreeView::UpdateAggregatedValues(FTableTreeNode& GroupNode)
+void STableTreeView::UpdateCStringSameValueAggregationSingleNode(const FTableColumn& InColumn, FTableTreeNode& InOutGroupNode)
+{
+	check(InOutGroupNode.IsGroup());
+
+	const TCHAR* AggregatedValue = nullptr;
+
+	// Find the first child node.
+	for (FBaseTreeNodePtr NodePtr : InOutGroupNode.GetFilteredChildren())
+	{
+		if (!NodePtr->IsGroup())
+		{
+			const TOptional<FTableCellValue> NodeValue = InColumn.GetValue(*NodePtr);
+			if (NodeValue.IsSet() &&
+				NodeValue.GetValue().DataType == ETableCellDataType::CString)
+			{
+				AggregatedValue = NodeValue.GetValue().CString;
+			}
+		}
+		else
+		{
+			FTableTreeNode& TableNode = *(FTableTreeNode*)NodePtr.Get();
+			if (TableNode.HasAggregatedValue(InColumn.GetId()))
+			{
+				const FTableCellValue& ChildGroupAggregatedValue = TableNode.GetAggregatedValue(InColumn.GetId());
+				if (ChildGroupAggregatedValue.DataType == ETableCellDataType::CString)
+				{
+					AggregatedValue = ChildGroupAggregatedValue.CString;
+				}
+			}
+		}
+
+		break;
+	}
+
+	if (AggregatedValue != nullptr)
+	{
+		// Check if all other children have the same value as the first node.
+		for (FBaseTreeNodePtr NodePtr : InOutGroupNode.GetFilteredChildren())
+		{
+			if (!NodePtr->IsGroup())
+			{
+				const TOptional<FTableCellValue> NodeValue = InColumn.GetValue(*NodePtr);
+				if (NodeValue.IsSet() &&
+					NodeValue.GetValue().DataType == ETableCellDataType::CString &&
+					AggregatedValue == NodeValue.GetValue().CString)
+				{
+					continue;
+				}
+			}
+			else
+			{
+				FTableTreeNode& TableNode = *(FTableTreeNode*)NodePtr.Get();
+				if (TableNode.HasAggregatedValue(InColumn.GetId()))
+				{
+					const FTableCellValue& ChildGroupAggregatedValue = TableNode.GetAggregatedValue(InColumn.GetId());
+					if (ChildGroupAggregatedValue.DataType == ETableCellDataType::CString &&
+						AggregatedValue == ChildGroupAggregatedValue.CString)
+					{
+						continue;
+					}
+				}
+			}
+
+			AggregatedValue = nullptr;
+			break;
+		}
+	}
+
+	if (AggregatedValue != nullptr)
+	{
+		InOutGroupNode.AddAggregatedValue(InColumn.GetId(), FTableCellValue(AggregatedValue));
+	}
+	else
+	{
+		InOutGroupNode.ResetAggregatedValue(InColumn.GetId());
+	}
+}
+
+////////////////////////////////////////////////////////////////////////////////////////////////////
+
+void STableTreeView::UpdateCStringSameValueAggregationRec(const FTableColumn& InColumn, FTableTreeNode& InOutGroupNode)
+{
+	check(InOutGroupNode.IsGroup());
+
+	// Update child group nodes first.
+	for (FBaseTreeNodePtr NodePtr : InOutGroupNode.GetFilteredChildren())
+	{
+		if (NodePtr->IsGroup())
+		{
+			FTableTreeNode& TableNode = *(FTableTreeNode*)NodePtr.Get();
+			STableTreeView::UpdateCStringSameValueAggregationRec(InColumn, TableNode);
+		}
+	}
+
+	STableTreeView::UpdateCStringSameValueAggregationSingleNode(InColumn, InOutGroupNode);
+}
+
+////////////////////////////////////////////////////////////////////////////////////////////////////
+
+template<typename T, bool bSetInitialValue, bool bIsRercursive>
+void STableTreeView::UpdateAggregation(const FTableColumn& InColumn, FTableTreeNode& InOutGroupNode, const T InitialAggregatedValue, TFunctionRef<T(T, const FTableCellValue&)> ValueGetterFunc)
+{
+	T AggregatedValue = InitialAggregatedValue;
+
+	for (FBaseTreeNodePtr NodePtr : InOutGroupNode.GetFilteredChildren())
+	{
+		if (NodePtr->IsGroup())
+		{
+			FTableTreeNode& TableNode = *(FTableTreeNode*)NodePtr.Get();
+			if (bIsRercursive)
+			{
+				TableNode.ResetAggregatedValue(InColumn.GetId());
+				STableTreeView::UpdateAggregation<T, bSetInitialValue, true>(InColumn, TableNode, InitialAggregatedValue, ValueGetterFunc);
+			}
+			if (TableNode.HasAggregatedValue(InColumn.GetId()))
+			{
+				AggregatedValue = ValueGetterFunc(AggregatedValue, TableNode.GetAggregatedValue(InColumn.GetId()));
+			}
+		}
+		else
+		{
+			const TOptional<FTableCellValue> NodeValue = InColumn.GetValue(*NodePtr);
+			if (NodeValue.IsSet())
+			{
+				AggregatedValue = ValueGetterFunc(AggregatedValue, NodeValue.GetValue());
+			}
+		}
+	}
+
+	if (bSetInitialValue || AggregatedValue != InitialAggregatedValue)
+	{
+		InOutGroupNode.AddAggregatedValue(InColumn.GetId(), FTableCellValue(AggregatedValue));
+	}
+}
+
+////////////////////////////////////////////////////////////////////////////////////////////////////
+
+template<bool bIsRercursive>
+void STableTreeView::UpdateAggregatedValues(TSharedPtr<FTable> InTable, FTableTreeNode& InOutGroupNode)
 {
 	FStopwatch Stopwatch;
 	Stopwatch.Start();
 
-	for (const TSharedRef<FTableColumn>& ColumnRef : Table->GetColumns())
+	check(InOutGroupNode.IsGroup());
+
+	for (const TSharedRef<FTableColumn>& ColumnRef : InTable->GetColumns())
 	{
 		FTableColumn& Column = ColumnRef.Get();
 		switch (Column.GetAggregation())
@@ -1588,7 +1756,7 @@ void STableTreeView::UpdateAggregatedValues(FTableTreeNode& GroupNode)
 			case ETableColumnAggregation::Sum:
 				if (Column.GetDataType() == ETableCellDataType::Float || Column.GetDataType() == ETableCellDataType::Double)
 				{
-					STableTreeView::UpdateAggregationRec<double>(Column, GroupNode, 0, true,
+					STableTreeView::UpdateAggregation<double, true, bIsRercursive>(Column, InOutGroupNode, 0.0,
 						[](double InValue, const FTableCellValue& InTableCellValue)
 						{
 							return InValue + InTableCellValue.AsDouble();
@@ -1596,7 +1764,7 @@ void STableTreeView::UpdateAggregatedValues(FTableTreeNode& GroupNode)
 				}
 				else
 				{
-					STableTreeView::UpdateAggregationRec<int64>(Column, GroupNode, 0, true,
+					STableTreeView::UpdateAggregation<int64, true, bIsRercursive>(Column, InOutGroupNode, (int64)0,
 						[](int64 InValue, const FTableCellValue& InTableCellValue)
 						{
 							return InValue + InTableCellValue.AsInt64();
@@ -1607,7 +1775,7 @@ void STableTreeView::UpdateAggregatedValues(FTableTreeNode& GroupNode)
 			case ETableColumnAggregation::Min:
 				if (Column.GetDataType() == ETableCellDataType::Float || Column.GetDataType() == ETableCellDataType::Double)
 				{
-					STableTreeView::UpdateAggregationRec<double>(Column, GroupNode, std::numeric_limits<double>::max(), false,
+					STableTreeView::UpdateAggregation<double, false, bIsRercursive>(Column, InOutGroupNode, std::numeric_limits<double>::max(),
 						[](double InValue, const FTableCellValue& InTableCellValue)
 						{
 							return FMath::Min(InValue, InTableCellValue.AsDouble());
@@ -1615,7 +1783,7 @@ void STableTreeView::UpdateAggregatedValues(FTableTreeNode& GroupNode)
 				}
 				else
 				{
-					STableTreeView::UpdateAggregationRec<int64>(Column, GroupNode, std::numeric_limits<int64>::max(), false,
+					STableTreeView::UpdateAggregation<int64, false, bIsRercursive>(Column, InOutGroupNode, std::numeric_limits<int64>::max(),
 						[](int64 InValue, const FTableCellValue& InTableCellValue)
 						{
 							return FMath::Min(InValue, InTableCellValue.AsInt64());
@@ -1626,7 +1794,7 @@ void STableTreeView::UpdateAggregatedValues(FTableTreeNode& GroupNode)
 			case ETableColumnAggregation::Max:
 				if (Column.GetDataType() == ETableCellDataType::Float || Column.GetDataType() == ETableCellDataType::Double)
 				{
-					STableTreeView::UpdateAggregationRec<double>(Column, GroupNode, std::numeric_limits<double>::lowest(), false,
+					STableTreeView::UpdateAggregation<double, false, bIsRercursive>(Column, InOutGroupNode, std::numeric_limits<double>::lowest(),
 						[](double InValue, const FTableCellValue& InTableCellValue)
 						{
 							return FMath::Max(InValue, InTableCellValue.AsDouble());
@@ -1634,7 +1802,7 @@ void STableTreeView::UpdateAggregatedValues(FTableTreeNode& GroupNode)
 				}
 				else
 				{
-					STableTreeView::UpdateAggregationRec<int64>(Column, GroupNode, std::numeric_limits<int64>::min(), false,
+					STableTreeView::UpdateAggregation<int64, false, bIsRercursive>(Column, InOutGroupNode, std::numeric_limits<int64>::min(),
 						[](int64 InValue, const FTableCellValue& InTableCellValue)
 						{
 							return FMath::Max(InValue, InTableCellValue.AsInt64());
@@ -1645,7 +1813,14 @@ void STableTreeView::UpdateAggregatedValues(FTableTreeNode& GroupNode)
 			case ETableColumnAggregation::SameValue:
 				if (Column.GetDataType() == ETableCellDataType::CString)
 				{
-					STableTreeView::UpdateCStringSameValueAggregationRec(Column, GroupNode);
+					if (bIsRercursive)
+					{
+						STableTreeView::UpdateCStringSameValueAggregationRec(Column, InOutGroupNode);
+					}
+					else
+					{
+						STableTreeView::UpdateCStringSameValueAggregationSingleNode(Column, InOutGroupNode);
+					}
 				}
 				break;
 		}
@@ -1655,110 +1830,22 @@ void STableTreeView::UpdateAggregatedValues(FTableTreeNode& GroupNode)
 	const double AggregationTime = Stopwatch.GetAccumulatedTime();
 	if (AggregationTime > 0.1)
 	{
-		UE_LOG(LogInsights, Log, TEXT("[Tree - %s] Aggregation completed in %.3fs."), *Table->GetDisplayName().ToString(), AggregationTime);
+		UE_LOG(LogInsights, Log, TEXT("[Tree - %s] Aggregation completed in %.3fs."), *InTable->GetDisplayName().ToString(), AggregationTime);
 	}
 }
 
 ////////////////////////////////////////////////////////////////////////////////////////////////////
 
-void STableTreeView::UpdateCStringSameValueAggregationRec(FTableColumn& Column, FTableTreeNode& GroupNode)
+void STableTreeView::UpdateAggregatedValuesSingleNode(FTableTreeNode& InOutGroupNode)
 {
-	// Update child group nodes first.
-	for (FBaseTreeNodePtr NodePtr : GroupNode.GetChildren())
-	{
-		if (NodePtr->IsFiltered())
-		{
-			continue;
-		}
-		if (NodePtr->IsGroup())
-		{
-			FTableTreeNode& TableNode = *(FTableTreeNode*)NodePtr.Get();
-			UpdateCStringSameValueAggregationRec(Column, TableNode);
-		}
-	}
-
-	const TCHAR* AggregatedValue = nullptr;
-
-	// Find the first child node.
-	for (FBaseTreeNodePtr NodePtr : GroupNode.GetChildren())
-	{
-		if (NodePtr->IsFiltered())
-		{
-			continue;
-		}
-
-		if (!NodePtr->IsGroup())
-		{
-			const TOptional<FTableCellValue> NodeValue = Column.GetValue(*NodePtr);
-			if (NodeValue.IsSet() &&
-				NodeValue.GetValue().DataType == ETableCellDataType::CString)
-			{
-				AggregatedValue = NodeValue.GetValue().CString;
-			}
-		}
-		else
-		{
-			FTableTreeNode& TableNode = *(FTableTreeNode*)NodePtr.Get();
-			if (TableNode.HasAggregatedValue(Column.GetId()))
-			{
-				const FTableCellValue& ChildGroupAggregatedValue = TableNode.GetAggregatedValue(Column.GetId());
-				if (ChildGroupAggregatedValue.DataType == ETableCellDataType::CString)
-				{
-					AggregatedValue = ChildGroupAggregatedValue.CString;
-				}
-			}
-		}
-
-		break;
-	}
-
-	if (AggregatedValue != nullptr)
-	{
-		// Check if all other children have the same value as the first node.
-		for (FBaseTreeNodePtr NodePtr : GroupNode.GetChildren())
-		{
-			if (NodePtr->IsFiltered())
-			{
-				continue;
-			}
-
-			if (!NodePtr->IsGroup())
-			{
-				const TOptional<FTableCellValue> NodeValue = Column.GetValue(*NodePtr);
-				if (NodeValue.IsSet() &&
-					NodeValue.GetValue().DataType == ETableCellDataType::CString &&
-					AggregatedValue == NodeValue.GetValue().CString)
-				{
-					continue;
-				}
-			}
-			else
-			{
-				FTableTreeNode& TableNode = *(FTableTreeNode*)NodePtr.Get();
-				if (TableNode.HasAggregatedValue(Column.GetId()))
-				{
-					const FTableCellValue& ChildGroupAggregatedValue = TableNode.GetAggregatedValue(Column.GetId());
-					if (ChildGroupAggregatedValue.DataType == ETableCellDataType::CString &&
-						AggregatedValue == ChildGroupAggregatedValue.CString)
-					{
-						continue;
-					}
-				}
-			}
+	STableTreeView::UpdateAggregatedValues<false>(Table, InOutGroupNode);
+}
 
-			AggregatedValue = nullptr;
-			break;
-		}
-	}
+////////////////////////////////////////////////////////////////////////////////////////////////////
 
-	if (AggregatedValue != nullptr)
-	{
-		GroupNode.AddAggregatedValue(Column.GetId(), FTableCellValue(AggregatedValue));
-	}
-	else
-	{
-		GroupNode.ResetAggregatedValues(Column.GetId());
-	}
+void STableTreeView::UpdateAggregatedValuesRec(FTableTreeNode& InOutGroupNode)
+{
+	STableTreeView::UpdateAggregatedValues<true>(Table, InOutGroupNode);
 }
 
 ////////////////////////////////////////////////////////////////////////////////////////////////////
diff --git a/Engine/Plugins/Editor/AssetManagerEditor/Source/AssetManagerEditor/Private/Insights/Table/Widgets/STableTreeView.h b/Engine/Plugins/Editor/AssetManagerEditor/Source/AssetManagerEditor/Private/Insights/Table/Widgets/STableTreeView.h
--- a/Engine/Plugins/Editor/AssetManagerEditor/Source/AssetManagerEditor/Private/Insights/Table/Widgets/STableTreeView.h
+++ b/Engine/Plugins/Editor/AssetManagerEditor/Source/AssetManagerEditor/Private/Insights/Table/Widgets/STableTreeView.h
@@ -243,6 +243,9 @@ protected:
 	/** Called by STreeView when selection has changed. */
 	void TreeView_OnSelectionChanged(FTableTreeNodePtr SelectedItem, ESelectInfo::Type SelectInfo);
 
+	/** Called by STreeView when a tree node is expanded or collapsed. */
+	void TreeView_OnExpansionChanged(FTableTreeNodePtr TreeNode, bool bShouldBeExpanded);
+
 	/** Called by STreeView when a tree item is double clicked. */
 	virtual void TreeView_OnMouseButtonDoubleClick(FTableTreeNodePtr TreeNode);
 
@@ -267,9 +270,8 @@ protected:
 
 	/** Populates the group and stat tree with items based on the current data. */
 	void ApplyFiltering();
-
+	void ApplyFiltering(FTableTreeNodePtr NodePtr);
 	bool ApplyAdvancedFiltersForNode(FTableTreeNodePtr NodePtr);
-
 	bool ApplyHierarchicalFilterForNode(FTableTreeNodePtr NodePtr, bool bFilterIsEmpty);
 
 	/** Set all the nodes belonging to a subtree as visible. Returns true if the caller node should be expanded. */
@@ -314,47 +316,17 @@ protected:
 	////////////////////////////////////////////////////////////////////////////////////////////////////
 	// Aggregation
 
-	void UpdateAggregatedValues(FTableTreeNode& GroupNode);
-
-	template<typename T>
-	static void UpdateAggregationRec(FTableColumn& Column, FTableTreeNode& GroupNode, T InitialAggregatedValue, bool bSetInitialValue, TFunctionRef<T(T, const FTableCellValue&)> ValueGetterFunc)
-	{
-		T AggregatedValue = InitialAggregatedValue;
+	static void UpdateCStringSameValueAggregationSingleNode(const FTableColumn& InColumn, FTableTreeNode& GroupNode);
+	static void UpdateCStringSameValueAggregationRec(const FTableColumn& InColumn, FTableTreeNode& GroupNode);
 
-		for (FBaseTreeNodePtr NodePtr : GroupNode.GetChildren())
-		{
-			if (NodePtr->IsFiltered())
-			{
-				continue;
-			}
-
-			if (!NodePtr->IsGroup())
-			{
-				const TOptional<FTableCellValue> NodeValue = Column.GetValue(*NodePtr);
-				if (NodeValue.IsSet())
-				{
-					AggregatedValue = ValueGetterFunc(AggregatedValue, NodeValue.GetValue());
-				}
-			}
-			else
-			{
-				FTableTreeNode& TableNode = *(FTableTreeNode*)NodePtr.Get();
-				TableNode.ResetAggregatedValues(Column.GetId());
-				UpdateAggregationRec(Column, TableNode, InitialAggregatedValue, bSetInitialValue, ValueGetterFunc);
-				if (TableNode.HasAggregatedValue(Column.GetId()))
-				{
-					AggregatedValue = ValueGetterFunc(AggregatedValue, TableNode.GetAggregatedValue(Column.GetId()));
-				}
-			}
-		}
+	template<typename T, bool bSetInitialValue, bool bIsRercursive>
+	static void UpdateAggregation(const FTableColumn& InColumn, FTableTreeNode& InOutGroupNode, const T InitialAggregatedValue, TFunctionRef<T(T, const FTableCellValue&)> ValueGetterFunc);
 
-		if (bSetInitialValue || InitialAggregatedValue != AggregatedValue)
-		{
-			GroupNode.AddAggregatedValue(Column.GetId(), FTableCellValue(AggregatedValue));
-		}
-	}
+	template<bool bIsRercursive>
+	static void UpdateAggregatedValues(TSharedPtr<FTable> InTable, FTableTreeNode& InOutGroupNode);
 
-	static void UpdateCStringSameValueAggregationRec(FTableColumn& Column, FTableTreeNode& GroupNode);
+	void UpdateAggregatedValuesSingleNode(FTableTreeNode& GroupNode);
+	void UpdateAggregatedValuesRec(FTableTreeNode& GroupNode);
 
 	////////////////////////////////////////////////////////////////////////////////////////////////////
 	// Sorting
diff --git a/Engine/Plugins/Editor/AssetManagerEditor/Source/AssetManagerEditor/Private/SAssetAuditBrowser.cpp b/Engine/Plugins/Editor/AssetManagerEditor/Source/AssetManagerEditor/Private/SAssetAuditBrowser.cpp
--- a/Engine/Plugins/Editor/AssetManagerEditor/Source/AssetManagerEditor/Private/SAssetAuditBrowser.cpp
+++ b/Engine/Plugins/Editor/AssetManagerEditor/Source/AssetManagerEditor/Private/SAssetAuditBrowser.cpp
@@ -25,6 +25,7 @@
 #include "Editor.h"
 
 #if ASSET_TABLE_TREE_VIEW_ENABLED
+#include "Insights/Common/Log.h"
 #include "Styling/StyleColors.h"
 #include "TreeView/AssetTable.h"
 #include "TreeView/SAssetTableTreeView.h"
@@ -732,20 +733,30 @@ void SAssetAuditBrowser::GoToHistoryIndex(int32 InHistoryIdx)
 #if ASSET_TABLE_TREE_VIEW_ENABLED
 void SAssetAuditBrowser::PopulateAssetTableRow(const FAssetData& AssetData, class FAssetTableRow* OutRow)
 {
-	OutRow->Type = AssetData.AssetClassPath.ToString();
-	OutRow->Name = AssetData.AssetName.ToString();
-	OutRow->Path = AssetData.GetObjectPathString();
+	OutRow->Type = AssetData.AssetClassPath.ToString(); // ???
+
+	//OutRow->PackageName = PackageName.ToString();     // ???
+	//OutRow->PackagePath = PackagePath.ToString();     // ???
+
+	OutRow->Name = AssetData.AssetName.ToString();      // vs. AssetData.AssetClassPath.GetAssetName().ToString() ???
+	OutRow->Path = AssetData.GetObjectPathString();     // vs. *AssetClassPath.ToString() ???
+
 	EditorModule->GetStringValueForCustomColumn(AssetData, FPrimaryAssetId::PrimaryAssetTypeTag, OutRow->PrimaryType);
 	EditorModule->GetStringValueForCustomColumn(AssetData, FPrimaryAssetId::PrimaryAssetNameTag, OutRow->PrimaryName);
+
 	//EditorModule->GetIntegerValueForCustomColumn(AssetData, IAssetManagerEditorModule::ManagedDiskSizeName, OutRow->ManagedDiskSize);
 	//EditorModule->GetIntegerValueForCustomColumn(AssetData, IAssetManagerEditorModule::DiskSizeName, OutRow->DiskSize);
+
 	EditorModule->GetIntegerValueForCustomColumn(AssetData, IAssetManagerEditorModule::StageChunkCompressedSizeName, OutRow->StagedCompressedSize);
 	EditorModule->GetIntegerValueForCustomColumn(AssetData, IAssetManagerEditorModule::TotalUsageName, OutRow->TotalUsageCount);
 	//EditorModule->GetStringValueForCustomColumn(AssetData, IAssetManagerEditorModule::CookRuleName, OutRow->CookRule);
 	//EditorModule->GetStringValueForCustomColumn(AssetData, IAssetManagerEditorModule::ChunksName, OutRow->Chunks);
+
 	EditorModule->GetStringValueForCustomColumn(AssetData, IAssetManagerEditorModule::GameFeaturePluginsName, OutRow->GameFeaturePlugin);
 	OutRow->NativeClass = AssetData.AssetClassPath.GetAssetName().ToString();
-	OutRow->Color = USlateThemeManager::Get().GetColor((EStyleColor)((int32)EStyleColor::AccentBlue + FMath::Rand() % 8)); // TODO: color by type!?
+
+	const uint32 Hash = GetTypeHash(OutRow->Type);
+	OutRow->Color = USlateThemeManager::Get().GetColor((EStyleColor)((uint32)EStyleColor::AccentBlue + Hash % 8));
 }
 #endif // ASSET_TABLE_TREE_VIEW_ENABLED
 
@@ -776,6 +787,11 @@ void SAssetAuditBrowser::RefreshAssetView()
 		TSharedPtr<FAssetTable> AssetTable = AssetTableTreeView->GetAssetTable();
 		if (AssetTable.IsValid())
 		{
+			UE_LOG(LogInsights, Log, TEXT("[AssetTree] Build asset table..."));
+
+			UE::Insights::FStopwatch Stopwatch;
+			Stopwatch.Start();
+
 			// Clears all tree nodes (that references the previous assets).
 			AssetTable->SetVisibleAssetCount(0);
 			AssetTableTreeView->RebuildTree(true);
@@ -816,7 +832,6 @@ void SAssetAuditBrowser::RefreshAssetView()
 					AlreadyProcessedPackages.Add(SourceAssets[SourceAssetIndex].PackageName);
 				}
 
-
 				FAssetIdentifier CurrentIdentifier(SourceAssets[SourceAssetIndex].PackageName);
 				// We'll have to add the dependencies to all these entries
 				TArray<FAssetData> AssetsInSourcePackage;
@@ -869,6 +884,11 @@ void SAssetAuditBrowser::RefreshAssetView()
 				}
 			}
 
+			Stopwatch.Stop();
+			const double TotalTime = Stopwatch.GetAccumulatedTime();
+			UE_LOG(LogInsights, Log, TEXT("[AssetTree] Asset table rebuilt in %.4fs (%d visible + %d hidden = %d assets)"),
+				TotalTime, AssetTable->GetVisibleAssetCount(), AssetTable->GetHiddenAssetCount(), AssetTable->GetTotalAssetCount());
+
 			AssetTableTreeView->RebuildTreeAsync();
 		}
 	}
diff --git a/Engine/Plugins/Editor/AssetManagerEditor/Source/AssetManagerEditor/Private/TreeView/AssetDependencyGrouping.cpp b/Engine/Plugins/Editor/AssetManagerEditor/Source/AssetManagerEditor/Private/TreeView/AssetDependencyGrouping.cpp
--- a/Engine/Plugins/Editor/AssetManagerEditor/Source/AssetManagerEditor/Private/TreeView/AssetDependencyGrouping.cpp
+++ b/Engine/Plugins/Editor/AssetManagerEditor/Source/AssetManagerEditor/Private/TreeView/AssetDependencyGrouping.cpp
@@ -8,6 +8,10 @@
 
 INSIGHTS_IMPLEMENT_RTTI(FAssetDependencyGrouping)
 
+////////////////////////////////////////////////////////////////////////////////////////////////////
+// FAssetDependencyGrouping
+////////////////////////////////////////////////////////////////////////////////////////////////////
+
 FAssetDependencyGrouping::FAssetDependencyGrouping()
 	: FTreeNodeGrouping(
 		LOCTEXT("Grouping_ByDependency_ShortName", "Dependency"),
@@ -18,10 +22,14 @@ FAssetDependencyGrouping::FAssetDependencyGrouping()
 {
 }
 
+////////////////////////////////////////////////////////////////////////////////////////////////////
+
 FAssetDependencyGrouping::~FAssetDependencyGrouping()
 {
 }
 
+////////////////////////////////////////////////////////////////////////////////////////////////////
+
 void FAssetDependencyGrouping::GroupNodes(const TArray<UE::Insights::FTableTreeNodePtr>& Nodes, UE::Insights::FTableTreeNode& ParentGroup, TWeakPtr<UE::Insights::FTable> InParentTable, UE::Insights::IAsyncOperationProgress& InAsyncOperationProgress) const
 {
 	using namespace UE::Insights;
@@ -31,9 +39,7 @@ void FAssetDependencyGrouping::GroupNodes(const TArray<UE::Insights::FTableTreeN
 	TSharedPtr<FAssetTable> AssetTable = StaticCastSharedPtr<FAssetTable>(InParentTable.Pin());
 	check(AssetTable.IsValid());
 
-	const FSlateBrush* IconBrush = FBaseTreeNode::GetDefaultIcon(true);
-	FLinearColor AssetGroupNodeColor(1.0f, 1.0f, 1.0f, 1.0f);
-	FLinearColor DependenciesGroupNodeColor(0.75f, 0.5f, 1.0f, 1.0f);
+	const FSlateBrush* AssetGroupNodeIconBrush = FBaseTreeNode::GetDefaultIcon(true);
 
 	for (FTableTreeNodePtr NodePtr : Nodes)
 	{
@@ -48,49 +54,69 @@ void FAssetDependencyGrouping::GroupNodes(const TArray<UE::Insights::FTableTreeN
 			continue;
 		}
 
-		// For each (visible) asset, we will create the following hierarcy:
+		// For each (visible) asset, we will create the following hierarchy:
 		//
+		// By default, the dependencies node is collapsed.
+		// |
+		// +-- [group:{AssetName}(Self+Deps)]               *[ASSET0]
+		// |   |
+		// |   +-- [group:_Self]                            *[ASSET0]
+		// |   |   |
+		// |   |   +-- [asset:{AssetName}]                  [ASSET0]
+		// |   |
+		// |   +-- [group:Dependencies]                     *[0]
+
+		// When dependencies node is expanded, it will be populated with actual dependencies.
 		// |
-		// +-- [group:{AssetName}]
+		// +-- [group:{AssetName}(Self+Deps)]               *[ASSET0+DEPS0']
 		// |   |
-		// |   +-- [asset:{AssetName}] // current asset
+		// |   +-- [group:_Self]                            *[ASSET0]
+		// |   |   |
+		// |   |   +-- [asset:{AssetName}]                  [ASSET0]
 		// |   |
-		// |   +-- [group:Dependencies]
+		// |   +-- [group:Dependencies]                     *[DEPS0']
+		// |       |
+		// |       +-- [asset:{DependentAsset1}]            [ASSET1]
 		// |       |
-		// |       +-- [asset:{DependentAsset1}] // dependent of current asset
+		// |       +-- [asset:{DependentAsset2}]            [ASSET2]
 		// |       |
-		// |       +-- [asset:{DependentAsset2}] // another dependent of current asset
+		// |       +-- [group:{DependentAsset3}(Self+Deps)] *[ASSET3+DEPS3']
 		// |       |
+		// |       ...
 
 		FAssetTreeNode& AssetNode = NodePtr->As<FAssetTreeNode>();
 		const FAssetTableRow& Asset = AssetNode.GetAssetChecked();
 
-		FTableTreeNodePtr GroupPtr = MakeShared<FCustomTableTreeNode>(FName(Asset.GetName(), 0), InParentTable, IconBrush, AssetGroupNodeColor);
-		GroupPtr->SetExpansion(false);
-		ParentGroup.AddChildAndSetGroupPtr(GroupPtr);
-
-		GroupPtr->AddChildAndSetGroupPtr(NodePtr);
-
-		FTableTreeNodePtr DepGroupPtr = nullptr;
-		const TArray<int32>& Dependencies = Asset.GetDependencies();
-		if (Dependencies.Num() > 0)
+		if (Asset.GetNumDependencies() > 0)
 		{
-			DepGroupPtr = MakeShared<FCustomTableTreeNode>(FName(TEXT("Dependencies"), 0), InParentTable, IconBrush, DependenciesGroupNodeColor);
-			DepGroupPtr->SetExpansion(false);
-			GroupPtr->AddChildAndSetGroupPtr(DepGroupPtr);
+			// Create a group for the asset node (self) + dependencies.
+			FTableTreeNodePtr AssetGroupPtr = MakeShared<FCustomTableTreeNode>(Asset.GetNodeName(), InParentTable, AssetGroupNodeIconBrush, Asset.GetColor());
+			AssetGroupPtr->SetExpansion(false);
+			ParentGroup.AddChildAndSetGroupPtr(AssetGroupPtr);
+
+			// Add the asset node (self) under a "_self" group node.
+			static FName SelfGroupName(TEXT("_Self_")); // used _ prefix to sort before "Dependencies"
+			FTableTreeNodePtr SelfGroupPtr = MakeShared<FCustomTableTreeNode>(SelfGroupName, InParentTable, AssetGroupNodeIconBrush, Asset.GetColor());
+			SelfGroupPtr->SetExpansion(false);
+			SelfGroupPtr->AddChildAndSetGroupPtr(NodePtr);
+			AssetGroupPtr->AddChildAndSetGroupPtr(SelfGroupPtr);
+
+			// Create a group node for all dependent assets of the current asset.
+			// The actual nodes for the dependent assets are lazy created by this group node.
+			static FName DependenciesGroupName(TEXT("Dependencies"));
+			TSharedPtr<FAssetDependenciesGroupTreeNode> DependenciesGroupPtr = MakeShared<FAssetDependenciesGroupTreeNode>(DependenciesGroupName, AssetTable, AssetNode.GetRowIndex());
+			DependenciesGroupPtr->SetExpansion(false);
+			AssetGroupPtr->AddChildAndSetGroupPtr(DependenciesGroupPtr);
 		}
-		for (int32 DepAssetIndex : Dependencies)
+		else
 		{
-			if (!ensure(AssetTable->IsValidRowIndex(DepAssetIndex)))
-			{
-				continue;
-			}
-			const FAssetTableRow& DepAsset = AssetTable->GetAssetChecked(DepAssetIndex);
-			FName DepAssetNodeName(DepAsset.GetName());
-			FAssetTreeNodePtr DepAssetNodePtr = MakeShared<FAssetTreeNode>(DepAssetNodeName, AssetTable, DepAssetIndex);
-			DepGroupPtr->AddChildAndSetGroupPtr(DepAssetNodePtr);
+			// No dependencies. Just add the asset node (self).
+			ParentGroup.AddChildAndSetGroupPtr(NodePtr);
 		}
 	}
 }
 
+////////////////////////////////////////////////////////////////////////////////////////////////////
+////////////////////////////////////////////////////////////////////////////////////////////////////
+
 #undef LOCTEXT_NAMESPACE
diff --git a/Engine/Plugins/Editor/AssetManagerEditor/Source/AssetManagerEditor/Private/TreeView/AssetDependencyGrouping.h b/Engine/Plugins/Editor/AssetManagerEditor/Source/AssetManagerEditor/Private/TreeView/AssetDependencyGrouping.h
--- a/Engine/Plugins/Editor/AssetManagerEditor/Source/AssetManagerEditor/Private/TreeView/AssetDependencyGrouping.h
+++ b/Engine/Plugins/Editor/AssetManagerEditor/Source/AssetManagerEditor/Private/TreeView/AssetDependencyGrouping.h
@@ -15,7 +15,4 @@ public:
 	virtual ~FAssetDependencyGrouping();
 
 	virtual void GroupNodes(const TArray<UE::Insights::FTableTreeNodePtr>& Nodes, UE::Insights::FTableTreeNode& ParentGroup, TWeakPtr<UE::Insights::FTable> InParentTable, UE::Insights::IAsyncOperationProgress& InAsyncOperationProgress) const override;
-
-private:
-	//...
 };
diff --git a/Engine/Plugins/Editor/AssetManagerEditor/Source/AssetManagerEditor/Private/TreeView/AssetTable.h b/Engine/Plugins/Editor/AssetManagerEditor/Source/AssetManagerEditor/Private/TreeView/AssetTable.h
--- a/Engine/Plugins/Editor/AssetManagerEditor/Source/AssetManagerEditor/Private/TreeView/AssetTable.h
+++ b/Engine/Plugins/Editor/AssetManagerEditor/Source/AssetManagerEditor/Private/TreeView/AssetTable.h
@@ -57,8 +57,12 @@ public:
 	//const TCHAR* GetCookRule() const { return *CookRule; }
 	//const TCHAR* GetChunks() const { return *Chunks; }
 	const TCHAR* GetNativeClass() const { return *NativeClass; }
-	const TArray<int32>& GetDependencies() const { return Dependencies; }
 	const TCHAR* GetGameFeaturePlugin() const { return *GameFeaturePlugin; }
+
+	int32 GetNumDependencies() const { return Dependencies.Num(); }
+	const TArray<int32>& GetDependencies() const { return Dependencies; }
+
+	FName GetNodeName() const { return FName(Name, 0); }
 	FLinearColor GetColor() const { return Color; }
 
 private:
@@ -74,8 +78,8 @@ private:
 	//FString CookRule;
 	//FString Chunks;
 	FString NativeClass;
-	TArray<int32> Dependencies;
 	FString GameFeaturePlugin;
+	TArray<int32> Dependencies;
 	FLinearColor Color;
 };
 
diff --git a/Engine/Plugins/Editor/AssetManagerEditor/Source/AssetManagerEditor/Private/TreeView/AssetTreeNode.cpp b/Engine/Plugins/Editor/AssetManagerEditor/Source/AssetManagerEditor/Private/TreeView/AssetTreeNode.cpp
--- a/Engine/Plugins/Editor/AssetManagerEditor/Source/AssetManagerEditor/Private/TreeView/AssetTreeNode.cpp
+++ b/Engine/Plugins/Editor/AssetManagerEditor/Source/AssetManagerEditor/Private/TreeView/AssetTreeNode.cpp
@@ -2,13 +2,21 @@
 
 #include "AssetTreeNode.h"
 
+#include "AssetDependencyGrouping.h"
+#include "Insights/Common/AsyncOperationProgress.h"
 #include "Insights/Common/InsightsStyle.h"
+#include "Insights/Table/Widgets/STableTreeView.h"
 
 INSIGHTS_IMPLEMENT_RTTI(FAssetTreeNode)
+INSIGHTS_IMPLEMENT_RTTI(FAssetDependenciesGroupTreeNode)
+
+////////////////////////////////////////////////////////////////////////////////////////////////////
+// FAssetTreeNode
+////////////////////////////////////////////////////////////////////////////////////////////////////
 
 const FSlateBrush* FAssetTreeNode::GetIcon() const
 {
-#if 1 // debug / mock code
+#if 0 // debug / mock code
 	if (RowId.RowIndex % 7)
 	{
 		return UE::Insights::FInsightsStyle::GetBrush("Icons.Package.TreeItem");
@@ -25,7 +33,53 @@ const FSlateBrush* FAssetTreeNode::GetIcon() const
 	return FBaseTreeNode::GetDefaultIcon(false); // default icon for leaf nodes
 }
 
+////////////////////////////////////////////////////////////////////////////////////////////////////
+
 FLinearColor FAssetTreeNode::GetColor() const
 {
 	return GetAssetChecked().GetColor();
 }
+
+////////////////////////////////////////////////////////////////////////////////////////////////////
+// FAssetDependenciesGroupTreeNode
+////////////////////////////////////////////////////////////////////////////////////////////////////
+
+FLinearColor FAssetDependenciesGroupTreeNode::GetColor() const
+{
+	return FLinearColor(0.75f, 0.5f, 1.0f, 1.0f);
+}
+
+////////////////////////////////////////////////////////////////////////////////////////////////////
+
+bool FAssetDependenciesGroupTreeNode::OnLazyCreateChildren(TSharedPtr<UE::Insights::STableTreeView> InTableTreeView)
+{
+	if (bAreChildrenCreated)
+	{
+		return false;
+	}
+
+	TSharedPtr<FAssetTable> AssetTable = StaticCastSharedPtr<FAssetTable>(GetParentTable().Pin());
+	const FAssetTableRow& AssetRow = AssetTable->GetAssetChecked(GetRowIndex());
+	const TArray<int32>& Dependencies = AssetRow.GetDependencies();
+	TArray<UE::Insights::FTableTreeNodePtr> AddedNodes;
+	for (int32 DepAssetIndex : Dependencies)
+	{
+		if (!ensure(AssetTable->IsValidRowIndex(DepAssetIndex)))
+		{
+			continue;
+		}
+		const FAssetTableRow& DepAssetRow = AssetTable->GetAssetChecked(DepAssetIndex);
+		FName DepAssetNodeName(DepAssetRow.GetName());
+		FAssetTreeNodePtr DepAssetNodePtr = MakeShared<FAssetTreeNode>(DepAssetNodeName, AssetTable, DepAssetIndex);
+		AddedNodes.Add(DepAssetNodePtr);
+	}
+
+	UE::Insights::FAsyncOperationProgress Progress;
+	FAssetDependencyGrouping Grouping;
+	Grouping.GroupNodes(AddedNodes, *this, GetParentTable(), Progress);
+
+	bAreChildrenCreated = true;
+	return true;
+}
+
+////////////////////////////////////////////////////////////////////////////////////////////////////
diff --git a/Engine/Plugins/Editor/AssetManagerEditor/Source/AssetManagerEditor/Private/TreeView/AssetTreeNode.h b/Engine/Plugins/Editor/AssetManagerEditor/Source/AssetManagerEditor/Private/TreeView/AssetTreeNode.h
--- a/Engine/Plugins/Editor/AssetManagerEditor/Source/AssetManagerEditor/Private/TreeView/AssetTreeNode.h
+++ b/Engine/Plugins/Editor/AssetManagerEditor/Source/AssetManagerEditor/Private/TreeView/AssetTreeNode.h
@@ -7,11 +7,19 @@
 #include "AssetTable.h"
 #include "Insights/Table/ViewModels/TableTreeNode.h"
 
+////////////////////////////////////////////////////////////////////////////////////////////////////
+
 class FAssetTreeNode : public UE::Insights::FTableTreeNode
 {
 	INSIGHTS_DECLARE_RTTI(FAssetTreeNode, UE::Insights::FTableTreeNode)
 
 public:
+	/** Hidden constructor. */
+	//explicit FAssetTreeNode(const FName InName, TWeakPtr<UE::Insights::FTable> InParentTable, int32 InRowIndex) = delete;
+
+	/** Hidden constructor. */
+	//explicit FAssetTreeNode(const FName InGroupName, TWeakPtr<UE::Insights::FTable> InParentTable) = delete;
+
 	/** Initialization constructor for the asset node. */
 	explicit FAssetTreeNode(const FName InName, TWeakPtr<FAssetTable> InParentTable, int32 InRowIndex)
 		: FTableTreeNode(InName, InParentTable, InRowIndex)
@@ -26,6 +34,8 @@ public:
 	{
 	}
 
+	virtual ~FAssetTreeNode() {}
+
 	FAssetTable& GetAssetTableChecked() const { return *AssetTable; }
 	const FAssetTableRow& GetAssetChecked() const { return AssetTable->GetAssetChecked(RowId.RowIndex); }
 
@@ -47,3 +57,40 @@ typedef TSharedRef<const class FAssetTreeNode> FAssetTreeNodeRefConst;
 
 /** Type definition for weak references to instances of FAssetTreeNode. */
 typedef TWeakPtr<class FAssetTreeNode> FAssetTreeNodeWeak;
+
+////////////////////////////////////////////////////////////////////////////////////////////////////
+
+class FAssetDependenciesGroupTreeNode : public UE::Insights::FTableTreeNode
+{
+	INSIGHTS_DECLARE_RTTI(FAssetDependenciesGroupTreeNode, UE::Insights::FTableTreeNode)
+
+public:
+	/** Hidden constructor. */
+	//explicit FAssetDependenciesGroupTreeNode(const FName InName, TWeakPtr<UE::Insights::FTable> InParentTable, int32 InRowIndex) = delete;
+
+	/** Hidden constructor. */
+	//explicit FAssetDependenciesGroupTreeNode(const FName InGroupName, TWeakPtr<UE::Insights::FTable> InParentTable) = delete;
+
+	/** Initialization constructor for the group node. */
+	explicit FAssetDependenciesGroupTreeNode(const FName InGroupName, TWeakPtr<FAssetTable> InParentTable, int32 InParentRowIndex)
+		: FTableTreeNode(InGroupName, InParentTable, InParentRowIndex)
+		, bAreChildrenCreated(false)
+	{
+		// This is actually a group node.
+		InitGroupData();
+
+		// Initially collapsed. Lazy create children when first expanded.
+		SetExpansion(false);
+	}
+
+	virtual ~FAssetDependenciesGroupTreeNode() {}
+
+	virtual FLinearColor GetColor() const override;
+
+	virtual bool OnLazyCreateChildren(TSharedPtr<class UE::Insights::STableTreeView> InTableTreeView) override;
+
+private:
+	bool bAreChildrenCreated;
+};
+
+////////////////////////////////////////////////////////////////////////////////////////////////////
