AssetManagerEditor: Fixed GetStringValueForCustomColumn(.., "GameFeaturePluginsName", ) to return true when value is getted correctly.

diff --git a/Engine/Plugins/Editor/AssetManagerEditor/Source/AssetManagerEditor/Private/AssetManagerEditorModule.cpp b/Engine/Plugins/Editor/AssetManagerEditor/Source/AssetManagerEditor/Private/AssetManagerEditorModule.cpp
--- a/Engine/Plugins/Editor/AssetManagerEditor/Source/AssetManagerEditor/Private/AssetManagerEditorModule.cpp
+++ b/Engine/Plugins/Editor/AssetManagerEditor/Source/AssetManagerEditor/Private/AssetManagerEditorModule.cpp
@@ -1308,10 +1308,12 @@ bool FAssetManagerEditorModule::GetStringValueForCustomColumn(const FAssetData&
 		if (PackageNameString.StartsWith(TEXT("/Game/")))
 		{
 			OutValue = TEXT("BaseGame");
+			return true;
 		}
 		else if (PackageNameString.ParseIntoArray(ParsedPath, TEXT("/")) >= 1)
 		{
 			OutValue = ParsedPath[0];
+			return true;
 		}
 	}
 	else
