AssetManagerEditor: Asset Tree View improvements and fixes.
  - Added "(double click to expand)" text suffix to "Dependencies" group node.
  - Fixed "Dependencies" group to create children only when expanded by user. This fixes UI freeze (infinite recursion) when changing/resetting the text filter.

diff --git a/Engine/Plugins/Editor/AssetManagerEditor/Source/AssetManagerEditor/Private/Insights/Table/Widgets/STableTreeView.h b/Engine/Plugins/Editor/AssetManagerEditor/Source/AssetManagerEditor/Private/Insights/Table/Widgets/STableTreeView.h
--- a/Engine/Plugins/Editor/AssetManagerEditor/Source/AssetManagerEditor/Private/Insights/Table/Widgets/STableTreeView.h
+++ b/Engine/Plugins/Editor/AssetManagerEditor/Source/AssetManagerEditor/Private/Insights/Table/Widgets/STableTreeView.h
@@ -127,6 +127,8 @@ public:
 	 */
 	void Construct(const FArguments& InArgs, TSharedPtr<FTable> InTablePtr);
 
+	TSharedPtr<STreeView<FTableTreeNodePtr>> GetInnerTreeView() const { return TreeView; }
+
 	TSharedPtr<FTable>& GetTable() { return Table; }
 	const TSharedPtr<FTable>& GetTable() const { return Table; }
 
diff --git a/Engine/Plugins/Editor/AssetManagerEditor/Source/AssetManagerEditor/Private/TreeView/AssetDependencyGrouping.cpp b/Engine/Plugins/Editor/AssetManagerEditor/Source/AssetManagerEditor/Private/TreeView/AssetDependencyGrouping.cpp
--- a/Engine/Plugins/Editor/AssetManagerEditor/Source/AssetManagerEditor/Private/TreeView/AssetDependencyGrouping.cpp
+++ b/Engine/Plugins/Editor/AssetManagerEditor/Source/AssetManagerEditor/Private/TreeView/AssetDependencyGrouping.cpp
@@ -58,29 +58,29 @@ void FAssetDependencyGrouping::GroupNodes(const TArray<UE::Insights::FTableTreeN
 		//
 		// By default, the dependencies node is collapsed.
 		// |
-		// +-- [group:{AssetName}(Self+Deps)]               *[ASSET0]
+		// +-- [group:{AssetName}] (self + dependencies)
 		// |   |
-		// |   +-- [group:_Self]                            *[ASSET0]
+		// |   +-- [group:_Self]
 		// |   |   |
-		// |   |   +-- [asset:{AssetName}]                  [ASSET0]
+		// |   |   +-- [asset:{AssetName}]
 		// |   |
-		// |   +-- [group:Dependencies]                     *[0]
+		// |   +-- [group:Dependencies] (double click to expand)
 
 		// When dependencies node is expanded, it will be populated with actual dependencies.
 		// |
-		// +-- [group:{AssetName}(Self+Deps)]               *[ASSET0+DEPS0']
+		// +-- [group:{AssetName}] (self + dependencies)
 		// |   |
-		// |   +-- [group:_Self]                            *[ASSET0]
+		// |   +-- [group:_Self]
 		// |   |   |
-		// |   |   +-- [asset:{AssetName}]                  [ASSET0]
+		// |   |   +-- [asset:{AssetName}]
 		// |   |
-		// |   +-- [group:Dependencies]                     *[DEPS0']
+		// |   +-- [group:Dependencies]
 		// |       |
-		// |       +-- [asset:{DependentAsset1}]            [ASSET1]
+		// |       +-- [asset:{DependentAsset1}]
 		// |       |
-		// |       +-- [asset:{DependentAsset2}]            [ASSET2]
+		// |       +-- [asset:{DependentAsset2}]
 		// |       |
-		// |       +-- [group:{DependentAsset3}(Self+Deps)] *[ASSET3+DEPS3']
+		// |       +-- [group:{DependentAsset3}] (self + dependencies)
 		// |       |
 		// |       ...
 
diff --git a/Engine/Plugins/Editor/AssetManagerEditor/Source/AssetManagerEditor/Private/TreeView/AssetTreeNode.cpp b/Engine/Plugins/Editor/AssetManagerEditor/Source/AssetManagerEditor/Private/TreeView/AssetTreeNode.cpp
--- a/Engine/Plugins/Editor/AssetManagerEditor/Source/AssetManagerEditor/Private/TreeView/AssetTreeNode.cpp
+++ b/Engine/Plugins/Editor/AssetManagerEditor/Source/AssetManagerEditor/Private/TreeView/AssetTreeNode.cpp
@@ -6,6 +6,9 @@
 #include "Insights/Common/AsyncOperationProgress.h"
 #include "Insights/Common/InsightsStyle.h"
 #include "Insights/Table/Widgets/STableTreeView.h"
+#include "Internationalization/Internationalization.h"
+
+#define LOCTEXT_NAMESPACE "FAssetTreeNode"
 
 INSIGHTS_IMPLEMENT_RTTI(FAssetTreeNode)
 INSIGHTS_IMPLEMENT_RTTI(FAssetDependenciesGroupTreeNode)
@@ -51,6 +54,17 @@ FLinearColor FAssetDependenciesGroupTreeNode::GetColor() const
 
 ////////////////////////////////////////////////////////////////////////////////////////////////////
 
+const FText FAssetDependenciesGroupTreeNode::GetExtraDisplayName() const
+{
+	if (!bAreChildrenCreated)
+	{
+		return LOCTEXT("DblClickToExpand", "(double click to expand)");
+	}
+	return FTableTreeNode::GetExtraDisplayName();
+}
+
+////////////////////////////////////////////////////////////////////////////////////////////////////
+
 bool FAssetDependenciesGroupTreeNode::OnLazyCreateChildren(TSharedPtr<UE::Insights::STableTreeView> InTableTreeView)
 {
 	if (bAreChildrenCreated)
@@ -58,6 +72,12 @@ bool FAssetDependenciesGroupTreeNode::OnLazyCreateChildren(TSharedPtr<UE::Insigh
 		return false;
 	}
 
+	TSharedPtr<STreeView<UE::Insights::FTableTreeNodePtr>> TreeView = InTableTreeView->GetInnerTreeView();
+	if (!TreeView || !TreeView->IsItemExpanded(SharedThis(this)))
+	{
+		return false;
+	}
+
 	TSharedPtr<FAssetTable> AssetTable = StaticCastSharedPtr<FAssetTable>(GetParentTable().Pin());
 	const FAssetTableRow& AssetRow = AssetTable->GetAssetChecked(GetRowIndex());
 	const TArray<int32>& Dependencies = AssetRow.GetDependencies();
@@ -83,3 +103,6 @@ bool FAssetDependenciesGroupTreeNode::OnLazyCreateChildren(TSharedPtr<UE::Insigh
 }
 
 ////////////////////////////////////////////////////////////////////////////////////////////////////
+////////////////////////////////////////////////////////////////////////////////////////////////////
+
+#undef LOCTEXT_NAMESPACE
diff --git a/Engine/Plugins/Editor/AssetManagerEditor/Source/AssetManagerEditor/Private/TreeView/AssetTreeNode.h b/Engine/Plugins/Editor/AssetManagerEditor/Source/AssetManagerEditor/Private/TreeView/AssetTreeNode.h
--- a/Engine/Plugins/Editor/AssetManagerEditor/Source/AssetManagerEditor/Private/TreeView/AssetTreeNode.h
+++ b/Engine/Plugins/Editor/AssetManagerEditor/Source/AssetManagerEditor/Private/TreeView/AssetTreeNode.h
@@ -85,6 +85,8 @@ public:
 
 	virtual ~FAssetDependenciesGroupTreeNode() {}
 
+	virtual const FText GetExtraDisplayName() const;
+
 	virtual FLinearColor GetColor() const override;
 
 	virtual bool OnLazyCreateChildren(TSharedPtr<class UE::Insights::STableTreeView> InTableTreeView) override;
