AssetManagerEditor: Asset Tree View improvements and fixes.
  - Added a string store implementation to de-duplicate the strings cached by FAssetTableRow.
  - Fixed asset path column to show only the actual path (excluding package name and asset name).

diff --git a/Engine/Plugins/Editor/AssetManagerEditor/Source/AssetManagerEditor/Private/SAssetAuditBrowser.cpp b/Engine/Plugins/Editor/AssetManagerEditor/Source/AssetManagerEditor/Private/SAssetAuditBrowser.cpp
--- a/Engine/Plugins/Editor/AssetManagerEditor/Source/AssetManagerEditor/Private/SAssetAuditBrowser.cpp
+++ b/Engine/Plugins/Editor/AssetManagerEditor/Source/AssetManagerEditor/Private/SAssetAuditBrowser.cpp
@@ -731,32 +731,72 @@ void SAssetAuditBrowser::GoToHistoryIndex(int32 InHistoryIdx)
 }
 
 #if ASSET_TABLE_TREE_VIEW_ENABLED
-void SAssetAuditBrowser::PopulateAssetTableRow(const FAssetData& AssetData, class FAssetTableRow* OutRow)
+void SAssetAuditBrowser::PopulateAssetTableRow(FAssetTableRow& OutRow, const FAssetData& AssetData, FAssetTable& AssetTable) const
 {
-	OutRow->Type = AssetData.AssetClassPath.ToString(); // ???
+	// Asset Type
+	FString AssetType;
+	if (!EditorModule->GetStringValueForCustomColumn(AssetData, FName(TEXT("Type")), AssetType)) // TODO: verify @Matt.Breindel
+	{
+		AssetType = AssetData.AssetClassPath.ToString(); // TODO: verify @Matt.Breindel
+	}
+	const uint32 AssetTypeHash = GetTypeHash(AssetType);
+	OutRow.Type = AssetTable.StoreStr(AssetType);
+
+	// Asset Name
+	OutRow.Name = AssetTable.StoreStr(AssetData.AssetName.ToString()); // TODO: verify @Matt.Breindel; see also AssetData.AssetClassPath.GetAssetName().ToString()
+
+	// Path (without package or asset name)
+	FString ObjectPathString = AssetData.GetObjectPathString(); // TODO: verify @Matt.Breindel; see also *AssetClassPath.ToString()
+	int Index;
+	ObjectPathString.FindLastChar(TEXT('/'), Index);
+	if (Index != INDEX_NONE)
+	{
+		ObjectPathString = ObjectPathString.Left(Index + 1);
+	}
+	OutRow.Path = AssetTable.StoreStr(ObjectPathString);
+
+	//OutRow.PackageName = AssetTable.StoreStr(AssetTable.PackageName.ToString()); // TODO: @Matt.Breindel, verify if this is needed/useful
+	//OutRow.PackagePath = AssetTable.StoreStr(AssetTable.PackagePath.ToString()); // TODO: @Matt.Breindel, verify if this is needed/useful
+
+	FString Str;
 
-	//OutRow->PackageName = PackageName.ToString();     // ???
-	//OutRow->PackagePath = PackagePath.ToString();     // ???
+	if (EditorModule->GetStringValueForCustomColumn(AssetData, FPrimaryAssetId::PrimaryAssetTypeTag, Str))
+	{
+		OutRow.PrimaryType = AssetTable.StoreStr(Str);
+	}
+	if (EditorModule->GetStringValueForCustomColumn(AssetData, FPrimaryAssetId::PrimaryAssetNameTag, Str))
+	{
+		OutRow.PrimaryName = AssetTable.StoreStr(Str);
+	}
 
-	OutRow->Name = AssetData.AssetName.ToString();      // vs. AssetData.AssetClassPath.GetAssetName().ToString() ???
-	OutRow->Path = AssetData.GetObjectPathString();     // vs. *AssetClassPath.ToString() ???
+	// TODO: @Matt.Breindel, verify if these columns are still needed
+	//EditorModule->GetIntegerValueForCustomColumn(AssetData, IAssetManagerEditorModule::ManagedDiskSizeName, OutRow.ManagedDiskSize);
+	//EditorModule->GetIntegerValueForCustomColumn(AssetData, IAssetManagerEditorModule::DiskSizeName, OutRow.DiskSize);
 
-	EditorModule->GetStringValueForCustomColumn(AssetData, FPrimaryAssetId::PrimaryAssetTypeTag, OutRow->PrimaryType);
-	EditorModule->GetStringValueForCustomColumn(AssetData, FPrimaryAssetId::PrimaryAssetNameTag, OutRow->PrimaryName);
+	EditorModule->GetIntegerValueForCustomColumn(AssetData, IAssetManagerEditorModule::StageChunkCompressedSizeName, OutRow.StagedCompressedSize);
+	EditorModule->GetIntegerValueForCustomColumn(AssetData, IAssetManagerEditorModule::TotalUsageName, OutRow.TotalUsageCount);
 
-	//EditorModule->GetIntegerValueForCustomColumn(AssetData, IAssetManagerEditorModule::ManagedDiskSizeName, OutRow->ManagedDiskSize);
-	//EditorModule->GetIntegerValueForCustomColumn(AssetData, IAssetManagerEditorModule::DiskSizeName, OutRow->DiskSize);
+	// TODO: @Matt.Breindel, verify if this column is still needed
+	//if (EditorModule->GetStringValueForCustomColumn(AssetData, IAssetManagerEditorModule::CookRuleName, Str))
+	//{
+	//	OutRow.CookRule = AssetTable.StoreStr(Str);
+	//}
 
-	EditorModule->GetIntegerValueForCustomColumn(AssetData, IAssetManagerEditorModule::StageChunkCompressedSizeName, OutRow->StagedCompressedSize);
-	EditorModule->GetIntegerValueForCustomColumn(AssetData, IAssetManagerEditorModule::TotalUsageName, OutRow->TotalUsageCount);
-	//EditorModule->GetStringValueForCustomColumn(AssetData, IAssetManagerEditorModule::CookRuleName, OutRow->CookRule);
-	//EditorModule->GetStringValueForCustomColumn(AssetData, IAssetManagerEditorModule::ChunksName, OutRow->Chunks);
+	// TODO: @Matt.Breindel, verify if this column is still needed
+	//if (EditorModule->GetStringValueForCustomColumn(AssetData, IAssetManagerEditorModule::ChunksName, Str))
+	//{
+	//	OutRow.Chunks = AssetTable.StoreStr(Str);
+	//}
 
-	EditorModule->GetStringValueForCustomColumn(AssetData, IAssetManagerEditorModule::GameFeaturePluginsName, OutRow->GameFeaturePlugin);
-	OutRow->NativeClass = AssetData.AssetClassPath.GetAssetName().ToString();
+	if (EditorModule->GetStringValueForCustomColumn(AssetData, IAssetManagerEditorModule::GameFeaturePluginsName, Str))
+	{
+		OutRow.GameFeaturePlugin = AssetTable.StoreStr(Str);
+	}
 
-	const uint32 Hash = GetTypeHash(OutRow->Type);
-	OutRow->Color = USlateThemeManager::Get().GetColor((EStyleColor)((uint32)EStyleColor::AccentBlue + Hash % 8));
+	OutRow.NativeClass = AssetTable.StoreStr(AssetData.AssetClassPath.GetAssetName().ToString());
+
+	// Sets the color based on asset type.
+	OutRow.Color = USlateThemeManager::Get().GetColor((EStyleColor)((uint32)EStyleColor::AccentBlue + AssetTypeHash % 8));
 }
 #endif // ASSET_TABLE_TREE_VIEW_ENABLED
 
@@ -807,7 +847,7 @@ void SAssetAuditBrowser::RefreshAssetView()
 			{
 				const FAssetData& SourceAsset = SourceAssets[Index];
 				FAssetTableRow AssetRow;
-				PopulateAssetTableRow(SourceAsset, &AssetRow);
+				PopulateAssetTableRow(AssetRow, SourceAsset, *AssetTable);
 				AssetTable->AddAsset(AssetRow);
 				AssetToIndexMap.Add(SourceAsset, Index);
 			}
@@ -860,7 +900,7 @@ void SAssetAuditBrowser::RefreshAssetView()
 							{
 								// We don't know about this asset yet. Create a new row for it and add it to the source list for further dependency analysis
 								FAssetTableRow NewRow;
-								PopulateAssetTableRow(DependencyAsset, &NewRow);
+								PopulateAssetTableRow(NewRow, DependencyAsset, *AssetTable);
 								AssetTable->AddAsset(NewRow);
 								IndicesOfDependenciesInRowTableToAddToCurrentSourceAssetRow.Add(SourceAssets.Num());
 								AssetToIndexMap.Add(DependencyAsset, SourceAssets.Num());
@@ -889,6 +929,12 @@ void SAssetAuditBrowser::RefreshAssetView()
 			UE_LOG(LogInsights, Log, TEXT("[AssetTree] Asset table rebuilt in %.4fs (%d visible + %d hidden = %d assets)"),
 				TotalTime, AssetTable->GetVisibleAssetCount(), AssetTable->GetHiddenAssetCount(), AssetTable->GetTotalAssetCount());
 
+			const FAssetTableStringStore& StringStore = AssetTable->GetStringStore();
+			UE_LOG(LogInsights, Log, TEXT("[AssetTree] String store: %d strings (%lld bytes) --> %d strings (%lld bytes | %lld bytes) %.0f%%"),
+				StringStore.GetNumInputStrings(), StringStore.GetTotalInputStringSize(), 
+				StringStore.GetNumStrings(), StringStore.GetTotalStringSize(), StringStore.GetAllocatedSize(),
+				(double)StringStore.GetAllocatedSize() * 100.0 / (double)StringStore.GetTotalInputStringSize());
+
 			AssetTableTreeView->RebuildTreeAsync();
 		}
 	}
diff --git a/Engine/Plugins/Editor/AssetManagerEditor/Source/AssetManagerEditor/Private/SAssetAuditBrowser.h b/Engine/Plugins/Editor/AssetManagerEditor/Source/AssetManagerEditor/Private/SAssetAuditBrowser.h
--- a/Engine/Plugins/Editor/AssetManagerEditor/Source/AssetManagerEditor/Private/SAssetAuditBrowser.h
+++ b/Engine/Plugins/Editor/AssetManagerEditor/Source/AssetManagerEditor/Private/SAssetAuditBrowser.h
@@ -133,7 +133,7 @@ protected:
 	void RefreshAssetView();
 
 #if ASSET_TABLE_TREE_VIEW_ENABLED
-	void PopulateAssetTableRow(const FAssetData& AssetData, class FAssetTableRow* OutRow);
+	void PopulateAssetTableRow(class FAssetTableRow& OutRow, const FAssetData& AssetData, class FAssetTable& AssetTable) const;
 #endif
 	/** Commands handled by this widget */
 	TSharedPtr<FUICommandList> Commands;
diff --git a/Engine/Plugins/Editor/AssetManagerEditor/Source/AssetManagerEditor/Private/TreeView/AssetTable.cpp b/Engine/Plugins/Editor/AssetManagerEditor/Source/AssetManagerEditor/Private/TreeView/AssetTable.cpp
--- a/Engine/Plugins/Editor/AssetManagerEditor/Source/AssetManagerEditor/Private/TreeView/AssetTable.cpp
+++ b/Engine/Plugins/Editor/AssetManagerEditor/Source/AssetManagerEditor/Private/TreeView/AssetTable.cpp
@@ -23,6 +23,8 @@ const FName FAssetTableColumns::CountColumnId(TEXT("Count"));
 const FName FAssetTableColumns::NameColumnId(TEXT("Name"));
 const FName FAssetTableColumns::TypeColumnId(TEXT("Type"));
 const FName FAssetTableColumns::PathColumnId(TEXT("Path"));
+//const FName FAssetTableColumns::PackageTypeColumnId(TEXT("PackageType"));
+//const FName FAssetTableColumns::PackageNameColumnId(TEXT("PackageName"));
 const FName FAssetTableColumns::PrimaryTypeColumnId(TEXT("PrimaryType"));
 const FName FAssetTableColumns::PrimaryNameColumnId(TEXT("PrimaryName"));
 //const FName FAssetTableColumns::ManagedDiskSizeColumnId(TEXT("ManagedDiskSize"));
@@ -34,6 +36,118 @@ const FName FAssetTableColumns::TotalUsageCountColumnId(TEXT("TotalUsageCount"))
 const FName FAssetTableColumns::NativeClassColumnId(TEXT("NativeClass"));
 const FName FAssetTableColumns::GameFeaturePluginColumnId(TEXT("GameFeaturePlugin"));
 
+////////////////////////////////////////////////////////////////////////////////////////////////////
+// FAssetTableStringStore
+////////////////////////////////////////////////////////////////////////////////////////////////////
+
+FAssetTableStringStore::FAssetTableStringStore()
+	: Chunks()
+	, Cache()
+	, TotalInputStringSize(0)
+	, TotalStoredStringSize(0)
+	, NumInputStrings(0)
+	, NumStoredStrings(0)
+{
+}
+
+////////////////////////////////////////////////////////////////////////////////////////////////////
+
+FAssetTableStringStore::~FAssetTableStringStore()
+{
+	Reset();
+}
+
+////////////////////////////////////////////////////////////////////////////////////////////////////
+
+void FAssetTableStringStore::Reset()
+{
+	for (FChunk& Chunk : Chunks)
+	{
+		delete[] Chunk.Buffer;
+	}
+	Chunks.Reset();
+	Cache.Reset();
+	TotalInputStringSize = 0;
+	TotalStoredStringSize = 0;
+	NumInputStrings = 0;
+	NumStoredStrings = 0;
+}
+
+////////////////////////////////////////////////////////////////////////////////////////////////////
+
+const FStringView FAssetTableStringStore::Store(const TCHAR* InStr)
+{
+	if (!InStr)
+	{
+		return FStringView();
+	}
+	return Store(FStringView(InStr));
+}
+
+////////////////////////////////////////////////////////////////////////////////////////////////////
+
+const FStringView FAssetTableStringStore::Store(const FStringView InStr)
+{
+	if (InStr.IsEmpty())
+	{
+		return FStringView();
+	}
+
+	check(InStr.Len() <= GetMaxStringLength());
+
+	TotalInputStringSize += (InStr.Len() + 1) * sizeof(TCHAR);
+	++NumInputStrings;
+
+	uint32 Hash = GetTypeHash(InStr);
+
+	TArray<FStringView> CachedStrings;
+	Cache.MultiFind(Hash, CachedStrings);
+	for (const FStringView& CachedString : CachedStrings)
+	{
+		if (CachedString.Equals(InStr, SearchCase))
+		{
+			return CachedString;
+		}
+	}
+
+	if (Chunks.Num() == 0 ||
+		Chunks.Last().Used + InStr.Len() + 1 > ChunkBufferLen)
+	{
+		AddChunk();
+	}
+
+	TotalStoredStringSize += (InStr.Len() + 1) * sizeof(TCHAR);
+	++NumStoredStrings;
+
+	FChunk& Chunk = Chunks.Last();
+	FStringView StoredStr(Chunk.Buffer + Chunk.Used, InStr.Len());
+	FMemory::Memcpy((void*)(Chunk.Buffer + Chunk.Used), (const void*)InStr.GetData(), InStr.Len() * sizeof(TCHAR));
+	Chunk.Used += InStr.Len();
+	Chunk.Buffer[Chunk.Used] = TEXT('\0');
+	Chunk.Used ++;
+	Cache.Add(Hash, StoredStr);
+	return StoredStr;
+}
+
+////////////////////////////////////////////////////////////////////////////////////////////////////
+
+void FAssetTableStringStore::AddChunk()
+{
+	FChunk& Chunk = Chunks.AddDefaulted_GetRef();
+	Chunk.Buffer = new TCHAR[ChunkBufferLen];
+	Chunk.Used = 0;
+}
+
+////////////////////////////////////////////////////////////////////////////////////////////////////
+
+void FAssetTableStringStore::EnumerateStrings(TFunction<void(const FStringView Str)> Callback) const
+{
+	for (const auto& Pair : Cache)
+	{
+		Callback(Pair.Value);
+	}
+}
+
 ////////////////////////////////////////////////////////////////////////////////////////////////////
 // FAssetTable
 ////////////////////////////////////////////////////////////////////////////////////////////////////
diff --git a/Engine/Plugins/Editor/AssetManagerEditor/Source/AssetManagerEditor/Private/TreeView/AssetTable.h b/Engine/Plugins/Editor/AssetManagerEditor/Source/AssetManagerEditor/Private/TreeView/AssetTable.h
--- a/Engine/Plugins/Editor/AssetManagerEditor/Source/AssetManagerEditor/Private/TreeView/AssetTable.h
+++ b/Engine/Plugins/Editor/AssetManagerEditor/Source/AssetManagerEditor/Private/TreeView/AssetTable.h
@@ -3,6 +3,9 @@
 #pragma once
 
 #include "CoreMinimal.h"
+#include "Containers/StringView.h"
+#include "Misc/CString.h"
+
 #include "Insights/Table/ViewModels/Table.h"
 
 namespace UE
@@ -13,13 +16,105 @@ namespace UE
 	}
 }
 
+// Simple string store.
+class FAssetTableStringStore
+{
+private:
+	static constexpr int32 ChunkBufferLen = 1024 * 1024; // 2 MiB/buffer
+	static constexpr ESearchCase::Type SearchCase = ESearchCase::CaseSensitive;
+
+	struct FChunk
+	{
+		TCHAR* Buffer = nullptr;
+		uint32 Used = 0;
+	};
+
+public:
+	FAssetTableStringStore();
+	~FAssetTableStringStore();
+
+	// Resets this string store. Frees all memory.
+	void Reset();
+
+	// Store a string. Gets the stored string for specified input string.
+	// @param InStr - the input string
+	// @returns the stored string view.
+	const FStringView Store(const TCHAR* InStr);
+
+	// Store a string. Gets the stored string for specified input string.
+	// The returned string is guaranteed to be null terminated.
+	// @param InStr - the input string
+	// @returns the stored string view.
+	const FStringView Store(const FStringView InStr);
+
+	// Gets the maximum length of a single string that can be stored.
+	// @returns the maximum string length
+	int32 GetMaxStringLength() const
+	{
+		return ChunkBufferLen - 1;
+	}
+
+	// Gets number of input strings, including duplicates.
+	// @returns number of input strings.
+	int32 GetNumInputStrings() const
+	{
+		return NumInputStrings;
+	}
+
+	// Gets size of input strings, including duplicates, in [bytes].
+	// @returns size of input strings, in [bytes].
+	SIZE_T GetTotalInputStringSize() const
+	{
+		return TotalInputStringSize;
+	}
+
+	// Gets number of stored (unique) strings.
+	// @returns number of stored strings.
+	int32 GetNumStrings() const
+	{
+		return NumStoredStrings;
+	}
+
+	// Gets size of stored (unique) strings, in [bytes].
+	// @returns size of stored strings, in [bytes].
+	SIZE_T GetTotalStringSize() const
+	{
+		return TotalStoredStringSize;
+	}
+
+	// Gets (estimated) total allocated memory for this string store, in [bytes].
+	// @returns memory size, in [bytes].
+	SIZE_T GetAllocatedSize() const
+	{
+		return Chunks.Num() * ChunkBufferLen * sizeof(TCHAR)
+			+ Chunks.GetAllocatedSize()
+			+ Cache.GetAllocatedSize()
+			+ sizeof(this);
+	}
+
+	void EnumerateStrings(TFunction<void(const FStringView Str)> Callback) const;
+
+private:
+	void AddChunk();
+
+private:
+	TArray<FChunk> Chunks;
+	TMultiMap<uint32, FStringView> Cache; // Hash --> FStringView
+	SIZE_T TotalInputStringSize;
+	SIZE_T TotalStoredStringSize;
+	int32 NumInputStrings;
+	int32 NumStoredStrings;
+};
+
 // Column identifiers
 struct FAssetTableColumns
 {
 	static const FName CountColumnId;
-	static const FName NameColumnId;
 	static const FName TypeColumnId;
+	static const FName NameColumnId;
 	static const FName PathColumnId;
+	//static const FName PackageNameColumnId;
+	//static const FName PackageTypeColumnId;
 	static const FName PrimaryTypeColumnId;
 	static const FName PrimaryNameColumnId;
 	//static const FName ManagedDiskSizeColumnId;
@@ -45,40 +140,42 @@ public:
 	{
 	}
 
-	const TCHAR* GetType() const { return *Type; }
-	const TCHAR* GetName() const { return *Name; }
-	const TCHAR* GetPath() const { return *Path; }
-	const TCHAR* GetPrimaryType() const { return *PrimaryType; }
-	const TCHAR* GetPrimaryName() const { return *PrimaryName; }
+	const TCHAR* GetType() const { return Type; }
+	const TCHAR* GetName() const { return Name; }
+	const TCHAR* GetPath() const { return Path; }
+	//const TCHAR* GetPackageName() const { return PackageName; }
+	//const TCHAR* GetPackagePath() const { return PackagePath; }
+	const TCHAR* GetPrimaryType() const { return PrimaryType; }
+	const TCHAR* GetPrimaryName() const { return PrimaryName; }
 	//int64 GetManagedDiskSize() const { return ManagedDiskSize; }
 	//int64 GetDiskSize() const { return DiskSize; }
 	int64 GetStagedCompressedSize() const { return StagedCompressedSize; }
 	int64 GetTotalUsageCount() const { return TotalUsageCount; }
-	//const TCHAR* GetCookRule() const { return *CookRule; }
-	//const TCHAR* GetChunks() const { return *Chunks; }
-	const TCHAR* GetNativeClass() const { return *NativeClass; }
-	const TCHAR* GetGameFeaturePlugin() const { return *GameFeaturePlugin; }
-
+	//const TCHAR* GetCookRule() const { return CookRule; }
+	//const TCHAR* GetChunks() const { return Chunks; }
+	const TCHAR* GetNativeClass() const { return NativeClass; }
+	const TCHAR* GetGameFeaturePlugin() const { return GameFeaturePlugin; }
 	int32 GetNumDependencies() const { return Dependencies.Num(); }
 	const TArray<int32>& GetDependencies() const { return Dependencies; }
-
 	FName GetNodeName() const { return FName(Name, 0); }
 	FLinearColor GetColor() const { return Color; }
 
 private:
-	FString Type;
-	FString Name;
-	FString Path;
-	FString PrimaryType;
-	FString PrimaryName;
+	const TCHAR* Type = nullptr;
+	const TCHAR* Name = nullptr;
+	const TCHAR* Path = nullptr;
+	//const TCHAR* PackageName = nullptr;
+	//const TCHAR* PackagePath = nullptr;
+	const TCHAR* PrimaryType = nullptr;
+	const TCHAR* PrimaryName = nullptr;
 	//int64 ManagedDiskSize = 0;
 	//int64 DiskSize = 0;
 	int64 StagedCompressedSize = 0;
 	int64 TotalUsageCount = 0;
-	//FString CookRule;
-	//FString Chunks;
-	FString NativeClass;
-	FString GameFeaturePlugin;
+	//const TCHAR* CookRule = nullptr;
+	//const TCHAR* Chunks = nullptr;
+	const TCHAR* NativeClass = nullptr;
+	const TCHAR* GameFeaturePlugin = nullptr;
 	TArray<int32> Dependencies;
 	FLinearColor Color;
 };
@@ -108,10 +205,16 @@ public:
 
 	void AddAsset(const FAssetTableRow& AssetRow) { Assets.Add(AssetRow); }
 
+	const FStringView StoreString(const FStringView InStr) { return StringStore.Store(InStr); }
+	const TCHAR* StoreStr(const FString& InStr) { return StringStore.Store(InStr).GetData(); }
+	const FAssetTableStringStore& GetStringStore() const { return StringStore; }
+	FAssetTableStringStore& GetStringStore() { return StringStore; }
+
 	void ClearAllData()
 	{
 		Assets.Empty();
 		VisibleAssetCount = 0;
+		StringStore.Reset();
 	}
 
 private:
@@ -120,4 +223,5 @@ private:
 private:
 	TArray<FAssetTableRow> Assets;
 	int32 VisibleAssetCount = 0;
+	FAssetTableStringStore StringStore;
 };
