AssetManagerEditor: Added "AssetAudit.EnableTreeView" console variable (false by default) to enable the tree view in Asset Audio UI (the next time the Asset Audit window is open).

diff --git a/Engine/Plugins/Editor/AssetManagerEditor/Source/AssetManagerEditor/Private/AssetManagerEditorModule.cpp b/Engine/Plugins/Editor/AssetManagerEditor/Source/AssetManagerEditor/Private/AssetManagerEditorModule.cpp
--- a/Engine/Plugins/Editor/AssetManagerEditor/Source/AssetManagerEditor/Private/AssetManagerEditorModule.cpp
+++ b/Engine/Plugins/Editor/AssetManagerEditor/Source/AssetManagerEditor/Private/AssetManagerEditorModule.cpp
@@ -59,10 +59,8 @@
 #include "ContentBrowserMenuContexts.h"
 #include "IContentBrowserDataModule.h"
 #include "ContentBrowserDataSubsystem.h"
-#if ASSET_TABLE_TREE_VIEW_ENABLED
 #include "Insights/Common/InsightsStyle.h"
 #include "Insights/Filter/ViewModels/Filters.h"
-#endif //ASSET_TABLE_TREE_VIEW_ENABLED
 
 #define LOCTEXT_NAMESPACE "AssetManagerEditor"
 
@@ -326,6 +324,9 @@ public:
 	virtual bool IsPackageInCurrentRegistrySource(FName PackageName) override;
 	virtual bool FilterAssetIdentifiersForCurrentRegistrySource(TArray<FAssetIdentifier>& AssetIdentifiers, const FAssetManagerDependencyQuery& DependencyQuery = FAssetManagerDependencyQuery::None(), bool bForwardDependency = true) override;
 	virtual bool WriteCollection(FName CollectionName, ECollectionShareType::Type ShareType, const TArray<FName>& PackageNames, bool bShowFeedback) override;
+
+	virtual bool ShouldEnableTreeViewInAssetAudit() const override;
+
 private:
 
 	static bool GetDependencyTypeArg(const FString& Arg, UE::AssetRegistry::EDependencyQuery& OutRequiredFlags);
@@ -401,6 +402,8 @@ private:
 	TSharedRef<SDockTab> SpawnAssetAuditTab(const FSpawnTabArgs& Args);
 	TSharedRef<SDockTab> SpawnReferenceViewerTab(const FSpawnTabArgs& Args);
 	TSharedRef<SDockTab> SpawnSizeMapTab(const FSpawnTabArgs& Args);
+
+	IConsoleVariable* EnableTreeViewInAssetAuditCVar = nullptr;
 };
 
 const TCHAR* FAssetManagerEditorModule::FindDepChainHelpText = TEXT("Finds all dependency chains from assets in the given search path, to the target package.\n Usage: FindDepChain TargetPackagePath SearchRootPath (optional: -hardonly/-softonly)\n e.g. FindDepChain /game/characters/heroes/muriel/meshes/muriel /game/cards ");
@@ -430,10 +433,14 @@ void FAssetManagerEditorModule::StartupModule()
 
 	if (GIsEditor && !IsRunningCommandlet())
 	{
-#if ASSET_TABLE_TREE_VIEW_ENABLED
 		UE::Insights::FInsightsStyle::Initialize();
 		UE::Insights::FFilterService::Initialize();
-#endif //ASSET_TABLE_TREE_VIEW_ENABLED
+
+		EnableTreeViewInAssetAuditCVar = IConsoleManager::Get().RegisterConsoleVariable(
+			TEXT("AssetAudit.EnableTreeView"),
+			false,
+			TEXT("Enables the tree view in Asset Audit (next time the window is open)."),
+			ECVF_Default);
 
 		AuditCmds.Add(IConsoleManager::Get().RegisterConsoleCommand(
 			TEXT("AssetManager.AssetAudit"),
@@ -548,6 +555,12 @@ void FAssetManagerEditorModule::ShutdownModule()
 	}
 	AuditCmds.Empty();
 
+	if (EnableTreeViewInAssetAuditCVar)
+	{
+		IConsoleManager::Get().UnregisterConsoleObject(EnableTreeViewInAssetAuditCVar);
+		EnableTreeViewInAssetAuditCVar = nullptr;
+	}
+
 	if ((GIsEditor && !IsRunningCommandlet()) && UObjectInitialized() && FSlateApplication::IsInitialized())
 	{
 		FContentBrowserModule& ContentBrowserModule = FModuleManager::LoadModuleChecked<FContentBrowserModule>(TEXT("ContentBrowser"));
@@ -610,10 +623,8 @@ void FAssetManagerEditorModule::ShutdownModule()
 		UToolMenus::UnRegisterStartupCallback(this);
 		UToolMenus::UnregisterOwner(this);
 
-#if ASSET_TABLE_TREE_VIEW_ENABLED
 		UE::Insights::FFilterService::Shutdown();
 		UE::Insights::FInsightsStyle::Shutdown();
-#endif //ASSET_TABLE_TREE_VIEW_ENABLED
 	}
 }
 
@@ -627,7 +638,6 @@ TSharedRef<SDockTab> FAssetManagerEditorModule::SpawnAssetAuditTab(const FSpawnT
 			SAssignNew(AssetAuditUI, SAssetAuditBrowser)
 		];
 
-#if ASSET_TABLE_TREE_VIEW_ENABLED
 	DockTab->SetOnTabClosed(SDockTab::FOnTabClosedCallback::CreateLambda([this](TSharedRef<SDockTab>)
 		{
 			if (AssetAuditUI.IsValid())
@@ -635,7 +645,6 @@ TSharedRef<SDockTab> FAssetManagerEditorModule::SpawnAssetAuditTab(const FSpawnT
 				AssetAuditUI.Pin()->OnClose();
 			}
 		}));
-#endif // ASSET_TABLE_TREE_VIEW_ENABLED
 
 	return DockTab;
 }
@@ -2455,4 +2464,9 @@ PRAGMA_ENABLE_DEPRECATION_WARNINGS
 	return bSuccess;
 }
 
+bool FAssetManagerEditorModule::ShouldEnableTreeViewInAssetAudit() const
+{
+	return EnableTreeViewInAssetAuditCVar && EnableTreeViewInAssetAuditCVar->GetBool();
+}
+
 #undef LOCTEXT_NAMESPACE
diff --git a/Engine/Plugins/Editor/AssetManagerEditor/Source/AssetManagerEditor/Private/SAssetAuditBrowser.cpp b/Engine/Plugins/Editor/AssetManagerEditor/Source/AssetManagerEditor/Private/SAssetAuditBrowser.cpp
--- a/Engine/Plugins/Editor/AssetManagerEditor/Source/AssetManagerEditor/Private/SAssetAuditBrowser.cpp
+++ b/Engine/Plugins/Editor/AssetManagerEditor/Source/AssetManagerEditor/Private/SAssetAuditBrowser.cpp
@@ -1,35 +1,32 @@
 // Copyright Epic Games, Inc. All Rights Reserved.
 
-
 #include "SAssetAuditBrowser.h"
-#include "AssetManagerEditorModule.h"
-#include "Framework/Commands/UICommandList.h"
-#include "Framework/MultiBox/MultiBoxBuilder.h"
 
-#include "Framework/Views/TableViewMetadata.h"
-#include "Widgets/Input/SButton.h"
-#include "FileHelpers.h"
+#include "AssetRegistry/AssetRegistryModule.h"
+#include "Blueprint/BlueprintSupport.h"
 #include "ClassViewerModule.h"
-#include "IContentBrowserSingleton.h"
-#include "ContentBrowserModule.h"
 #include "ContentBrowserDataSource.h"
-#include "AssetRegistry/AssetRegistryModule.h"
-#include "Toolkits/GlobalEditorCommonCommands.h"
+#include "ContentBrowserModule.h"
+#include "DragAndDrop/AssetDragDropOp.h"
+#include "Editor.h"
 #include "Engine/AssetManager.h"
-#include "AssetManagerEditorCommands.h"
 #include "Engine/BlueprintCore.h"
-#include "Widgets/Input/SComboBox.h"
-#include "DragAndDrop/AssetDragDropOp.h"
-#include "Blueprint/BlueprintSupport.h"
+#include "FileHelpers.h"
+#include "Framework/Commands/UICommandList.h"
+#include "Framework/MultiBox/MultiBoxBuilder.h"
+#include "Framework/Views/TableViewMetadata.h"
+#include "IContentBrowserSingleton.h"
+#include "Styling/StyleColors.h"
 #include "Subsystems/AssetEditorSubsystem.h"
-#include "Editor.h"
+#include "Toolkits/GlobalEditorCommonCommands.h"
+#include "Widgets/Input/SButton.h"
+#include "Widgets/Input/SComboBox.h"
 
-#if ASSET_TABLE_TREE_VIEW_ENABLED
+#include "AssetManagerEditorModule.h"
+#include "AssetManagerEditorCommands.h"
 #include "Insights/Common/Log.h"
-#include "Styling/StyleColors.h"
 #include "TreeView/AssetTable.h"
 #include "TreeView/SAssetTableTreeView.h"
-#endif
 
 #define LOCTEXT_NAMESPACE "AssetManagementBrowser"
 
@@ -359,22 +356,7 @@ void SAssetAuditBrowser::Construct(const FArguments& InArgs)
 			]
 		];
 
-#if ASSET_TABLE_TREE_VIEW_ENABLED
-	TSharedRef<FAssetTable> AssetTable = MakeShared<FAssetTable>();
-	AssetTable->Reset();
-	AssetTable->SetDisplayName(FText::FromString(TEXT("AssetTable")));
-#endif //ASSET_TABLE_TREE_VIEW_ENABLED
-
-	this->ChildSlot
-	[
-#if ASSET_TABLE_TREE_VIEW_ENABLED
-	SNew(SSplitter)
-	.Orientation(EOrientation::Orient_Vertical)
-	+ SSplitter::Slot()
-	.Value(0.5f)
-	.MinSize(100.0f)
-	[
-#endif // ASSET_TABLE_TREE_VIEW_ENABLED
+	TSharedRef<SWidget> DefaultContent =
 		SNew(SVerticalBox)
 		+SVerticalBox::Slot()
 		.AutoHeight()
@@ -525,19 +507,39 @@ void SAssetAuditBrowser::Construct(const FArguments& InArgs)
 			[
 				ContentBrowserModule.Get().CreateAssetPicker(Config)
 			]
-		]
-#if ASSET_TABLE_TREE_VIEW_ENABLED
-	]
-	+ SSplitter::Slot()
-	[
-		SAssignNew(AssetTableTreeView, SAssetTableTreeView, AssetTable)
-	]
-#endif //ASSET_TABLE_TREE_VIEW_ENABLED
-	];
-
-#if ASSET_TABLE_TREE_VIEW_ENABLED
-	AssetTableTreeView->RebuildTreeAsync(); // only needed for initial mock data
-#endif
+		];
+
+	if (ManagerEditorModule.ShouldEnableTreeViewInAssetAudit())
+	{
+		TSharedRef<FAssetTable> AssetTable = MakeShared<FAssetTable>();
+		AssetTable->Reset();
+		AssetTable->SetDisplayName(FText::FromString(TEXT("AssetTable")));
+
+		this->ChildSlot
+		[
+			SNew(SSplitter)
+			.Orientation(EOrientation::Orient_Vertical)
+			+ SSplitter::Slot()
+			.Value(0.5f)
+			.MinSize(100.0f)
+			[
+				DefaultContent
+			]
+			+ SSplitter::Slot()
+			[
+				SAssignNew(AssetTableTreeView, SAssetTableTreeView, AssetTable)
+			]
+		];
+
+		AssetTableTreeView->RebuildTreeAsync(); // only needed for initial mock data
+	}
+	else
+	{
+		this->ChildSlot
+		[
+			DefaultContent
+		];
+	}
 
 	// Refresh our data sources to make sure we are displaying up to date info otherwise the user will need to 
 	// hit "refresh" every time they open the dialog to be able to trust it
@@ -730,7 +732,6 @@ void SAssetAuditBrowser::GoToHistoryIndex(int32 InHistoryIdx)
 	}
 }
 
-#if ASSET_TABLE_TREE_VIEW_ENABLED
 void SAssetAuditBrowser::PopulateAssetTableRow(FAssetTableRow& OutRow, const FAssetData& AssetData, FAssetTable& AssetTable) const
 {
 	// Asset Type
@@ -798,7 +799,6 @@ void SAssetAuditBrowser::PopulateAssetTableRow(FAssetTableRow& OutRow, const FAs
 	// Sets the color based on asset type.
 	OutRow.Color = USlateThemeManager::Get().GetColor((EStyleColor)((uint32)EStyleColor::AccentBlue + AssetTypeHash % 8));
 }
-#endif // ASSET_TABLE_TREE_VIEW_ENABLED
 
 void SAssetAuditBrowser::RefreshAssetView()
 {
@@ -820,7 +820,6 @@ void SAssetAuditBrowser::RefreshAssetView()
 
 	SetFilterDelegate.Execute(Filter);
 
-#if ASSET_TABLE_TREE_VIEW_ENABLED
 	// Refresh list of assets for the tree view
 	if (AssetTableTreeView.IsValid())
 	{
@@ -938,7 +937,6 @@ void SAssetAuditBrowser::RefreshAssetView()
 			AssetTableTreeView->RebuildTreeAsync();
 		}
 	}
-#endif // ASSET_TABLE_TREE_VIEW_ENABLED
 }
 
 FReply SAssetAuditBrowser::ClearAssets()
@@ -1136,12 +1134,10 @@ void SAssetAuditBrowser::SetCurrentRegistrySource(const FAssetManagerEditorRegis
 
 void SAssetAuditBrowser::OnClose()
 {
-#if ASSET_TABLE_TREE_VIEW_ENABLED
 	if (AssetTableTreeView.IsValid())
 	{
 		AssetTableTreeView->OnClose();
 	}
-#endif // ASSET_TABLE_TREE_VIEW_ENABLED
 }
 
 #undef LOCTEXT_NAMESPACE
diff --git a/Engine/Plugins/Editor/AssetManagerEditor/Source/AssetManagerEditor/Private/SAssetAuditBrowser.h b/Engine/Plugins/Editor/AssetManagerEditor/Source/AssetManagerEditor/Private/SAssetAuditBrowser.h
--- a/Engine/Plugins/Editor/AssetManagerEditor/Source/AssetManagerEditor/Private/SAssetAuditBrowser.h
+++ b/Engine/Plugins/Editor/AssetManagerEditor/Source/AssetManagerEditor/Private/SAssetAuditBrowser.h
@@ -5,8 +5,6 @@
 #include "ContentBrowserDelegates.h"
 #include "Widgets/SCompoundWidget.h"
 
-#define ASSET_TABLE_TREE_VIEW_ENABLED 0 // switch for work in progress / experimental code; to be removed once ready
-
 struct FAssetIdentifier;
 struct FAssetManagerEditorRegistrySource;
 
@@ -132,9 +130,9 @@ protected:
 	/** Refresh the asset view with a new filter */
 	void RefreshAssetView();
 
-#if ASSET_TABLE_TREE_VIEW_ENABLED
+	/** Populates an asset table row (view model for the asset tree view) with asset data. */
 	void PopulateAssetTableRow(class FAssetTableRow& OutRow, const FAssetData& AssetData, class FAssetTable& AssetTable) const;
-#endif
+
 	/** Commands handled by this widget */
 	TSharedPtr<FUICommandList> Commands;
 
@@ -175,7 +173,6 @@ protected:
 	class UAssetManager* AssetManager;
 	class IAssetManagerEditorModule* EditorModule;
 
-#if ASSET_TABLE_TREE_VIEW_ENABLED
+	/** The Tree View widget */
 	TSharedPtr<class SAssetTableTreeView> AssetTableTreeView;
-#endif
 };
diff --git a/Engine/Plugins/Editor/AssetManagerEditor/Source/AssetManagerEditor/Public/AssetManagerEditorModule.h b/Engine/Plugins/Editor/AssetManagerEditor/Source/AssetManagerEditor/Public/AssetManagerEditorModule.h
--- a/Engine/Plugins/Editor/AssetManagerEditor/Source/AssetManagerEditor/Public/AssetManagerEditorModule.h
+++ b/Engine/Plugins/Editor/AssetManagerEditor/Source/AssetManagerEditor/Public/AssetManagerEditorModule.h
@@ -423,6 +423,9 @@ public:
 	virtual void OpenAssetAuditUI(TArray<FAssetIdentifier> SelectedIdentifiers) = 0;
 	virtual void OpenAssetAuditUI(TArray<FName> SelectedPackages) = 0;
 
+	/** Returns true if the tree view will be enabled in the Asset Audit UI, the next time the Asset Audit window is open. */
+	virtual bool ShouldEnableTreeViewInAssetAudit() const = 0;
+
 	/** Spawns reference viewer, showing selected packages or identifiers */
 	virtual void OpenReferenceViewerUI(const TArray<FAssetIdentifier> SelectedIdentifiers, const FReferenceViewerParams ReferenceViewerParams = FReferenceViewerParams()) = 0;
 	virtual void OpenReferenceViewerUI(const TArray<FName> SelectedPackages, const FReferenceViewerParams ReferenceViewerParams = FReferenceViewerParams()) = 0;
@@ -435,7 +438,6 @@ public:
 	virtual void OpenSizeMapUI(TArray<FAssetIdentifier> SelectedIdentifiers) = 0;
 	virtual void OpenSizeMapUI(TArray<FName> SelectedPackages) = 0;
 
-
 	/** Open the Shader cook stats */
 	virtual void OpenShaderCookStatistics(TArray<FName> SelectedIdentifiers) = 0;
 };
