Fixed static analysis warning.

diff --git a/Engine/Plugins/Editor/AssetManagerEditor/Source/AssetManagerEditor/Private/TreeView/AssetTable.h b/Engine/Plugins/Editor/AssetManagerEditor/Source/AssetManagerEditor/Private/TreeView/AssetTable.h
--- a/Engine/Plugins/Editor/AssetManagerEditor/Source/AssetManagerEditor/Private/TreeView/AssetTable.h
+++ b/Engine/Plugins/Editor/AssetManagerEditor/Source/AssetManagerEditor/Private/TreeView/AssetTable.h
@@ -88,8 +88,7 @@ public:
 	{
 		return Chunks.Num() * ChunkBufferLen * sizeof(TCHAR)
 			+ Chunks.GetAllocatedSize()
-			+ Cache.GetAllocatedSize()
-			+ sizeof(this);
+			+ Cache.GetAllocatedSize();
 	}
 
 	void EnumerateStrings(TFunction<void(const FStringView Str)> Callback) const;
