Remove restriction that Reference Viewer cannot be used on read-only assets

diff --git a/Engine/Plugins/Editor/AssetManagerEditor/Source/AssetManagerEditor/Private/AssetManagerEditorModule.cpp b/Engine/Plugins/Editor/AssetManagerEditor/Source/AssetManagerEditor/Private/AssetManagerEditorModule.cpp
--- a/Engine/Plugins/Editor/AssetManagerEditor/Source/AssetManagerEditor/Private/AssetManagerEditorModule.cpp
+++ b/Engine/Plugins/Editor/AssetManagerEditor/Source/AssetManagerEditor/Private/AssetManagerEditorModule.cpp
@@ -983,7 +983,7 @@ void FAssetManagerEditorModule::ExtendContentBrowserAssetSelectionMenu()
 	FToolMenuEntry& Entry = Section.AddDynamicEntry("AssetManagerEditorViewCommands", FNewToolMenuSectionDelegate::CreateLambda([this](FToolMenuSection& InSection)
 	{
 		UContentBrowserAssetContextMenuContext* Context = InSection.FindContext<UContentBrowserAssetContextMenuContext>();
-		if (Context && Context->bCanBeModified && Context->SelectedAssets.Num() > 0)
+		if (Context && Context->SelectedAssets.Num() > 0)
 		{
 			CreateAssetContextMenu(InSection);
 		}
@@ -997,7 +997,7 @@ void FAssetManagerEditorModule::ExtendContentBrowserPathSelectionMenu()
 	FToolMenuEntry& Entry = Section.AddDynamicEntry("AssetManagerEditorViewCommands", FNewToolMenuSectionDelegate::CreateLambda([this](FToolMenuSection& InSection)
 	{
 		UContentBrowserFolderContext* Context = InSection.FindContext<UContentBrowserFolderContext>();
-		if (Context && Context->bCanBeModified && Context->NumAssetPaths > 0)
+		if (Context && Context->NumAssetPaths > 0)
 		{
 			CreateAssetContextMenu(InSection);
 		}
