Adding dependency size data to the asset table tree view.

diff --git a/Engine/Plugins/Editor/AssetManagerEditor/Source/AssetManagerEditor/Private/SAssetAuditBrowser.cpp b/Engine/Plugins/Editor/AssetManagerEditor/Source/AssetManagerEditor/Private/SAssetAuditBrowser.cpp
--- a/Engine/Plugins/Editor/AssetManagerEditor/Source/AssetManagerEditor/Private/SAssetAuditBrowser.cpp
+++ b/Engine/Plugins/Editor/AssetManagerEditor/Source/AssetManagerEditor/Private/SAssetAuditBrowser.cpp
@@ -30,8 +30,6 @@
 
 #define LOCTEXT_NAMESPACE "AssetManagementBrowser"
 
-//#pragma optimize("", off)
-
 const FString SAssetAuditBrowser::SettingsIniSection = TEXT("AssetManagementBrowser");
 const int32 SAssetAuditBrowser::MaxAssetsHistory = 10;
 
@@ -736,18 +734,18 @@ void SAssetAuditBrowser::PopulateAssetTableRow(FAssetTableRow& OutRow, const FAs
 {
 	// Asset Type
 	FString AssetType;
-	if (!EditorModule->GetStringValueForCustomColumn(AssetData, FName(TEXT("Type")), AssetType)) // TODO: verify @Matt.Breindel
+	if (!EditorModule->GetStringValueForCustomColumn(AssetData, FName(TEXT("Type")), AssetType)) 
 	{
-		AssetType = AssetData.AssetClassPath.ToString(); // TODO: verify @Matt.Breindel
+		AssetType = AssetData.AssetClassPath.ToString(); 
 	}
 	const uint32 AssetTypeHash = GetTypeHash(AssetType);
 	OutRow.Type = AssetTable.StoreStr(AssetType);
 
 	// Asset Name
-	OutRow.Name = AssetTable.StoreStr(AssetData.AssetName.ToString()); // TODO: verify @Matt.Breindel; see also AssetData.AssetClassPath.GetAssetName().ToString()
+	OutRow.Name = AssetTable.StoreStr(AssetData.AssetName.ToString());
 
 	// Path (without package or asset name)
-	FString ObjectPathString = AssetData.GetObjectPathString(); // TODO: verify @Matt.Breindel; see also *AssetClassPath.ToString()
+	FString ObjectPathString = AssetData.GetObjectPathString();
 	int Index;
 	ObjectPathString.FindLastChar(TEXT('/'), Index);
 	if (Index != INDEX_NONE)
@@ -756,9 +754,6 @@ void SAssetAuditBrowser::PopulateAssetTableRow(FAssetTableRow& OutRow, const FAs
 	}
 	OutRow.Path = AssetTable.StoreStr(ObjectPathString);
 
-	//OutRow.PackageName = AssetTable.StoreStr(AssetTable.PackageName.ToString()); // TODO: @Matt.Breindel, verify if this is needed/useful
-	//OutRow.PackagePath = AssetTable.StoreStr(AssetTable.PackagePath.ToString()); // TODO: @Matt.Breindel, verify if this is needed/useful
-
 	FString Str;
 
 	if (EditorModule->GetStringValueForCustomColumn(AssetData, FPrimaryAssetId::PrimaryAssetTypeTag, Str))
@@ -770,24 +765,13 @@ void SAssetAuditBrowser::PopulateAssetTableRow(FAssetTableRow& OutRow, const FAs
 		OutRow.PrimaryName = AssetTable.StoreStr(Str);
 	}
 
-	// TODO: @Matt.Breindel, verify if these columns are still needed
-	//EditorModule->GetIntegerValueForCustomColumn(AssetData, IAssetManagerEditorModule::ManagedDiskSizeName, OutRow.ManagedDiskSize);
-	//EditorModule->GetIntegerValueForCustomColumn(AssetData, IAssetManagerEditorModule::DiskSizeName, OutRow.DiskSize);
-
 	EditorModule->GetIntegerValueForCustomColumn(AssetData, IAssetManagerEditorModule::StageChunkCompressedSizeName, OutRow.StagedCompressedSize);
 	EditorModule->GetIntegerValueForCustomColumn(AssetData, IAssetManagerEditorModule::TotalUsageName, OutRow.TotalUsageCount);
 
-	// TODO: @Matt.Breindel, verify if this column is still needed
-	//if (EditorModule->GetStringValueForCustomColumn(AssetData, IAssetManagerEditorModule::CookRuleName, Str))
-	//{
-	//	OutRow.CookRule = AssetTable.StoreStr(Str);
-	//}
-
-	// TODO: @Matt.Breindel, verify if this column is still needed
-	//if (EditorModule->GetStringValueForCustomColumn(AssetData, IAssetManagerEditorModule::ChunksName, Str))
-	//{
-	//	OutRow.Chunks = AssetTable.StoreStr(Str);
-	//}
+	if (EditorModule->GetStringValueForCustomColumn(AssetData, IAssetManagerEditorModule::ChunksName, Str))
+	{
+		OutRow.Chunks = AssetTable.StoreStr(Str);
+	}
 
 	if (EditorModule->GetStringValueForCustomColumn(AssetData, IAssetManagerEditorModule::GameFeaturePluginsName, Str))
 	{
@@ -820,6 +804,7 @@ void SAssetAuditBrowser::RefreshAssetView()
 
 	SetFilterDelegate.Execute(Filter);
 
+
 	// Refresh list of assets for the tree view
 	if (AssetTableTreeView.IsValid())
 	{
@@ -841,19 +826,38 @@ void SAssetAuditBrowser::RefreshAssetView()
 			TMap<FAssetData, int32> AssetToIndexMap;
 
 			TArray<FAssetData> SourceAssets;
-			AssetRegistry->GetAssets(Filter, SourceAssets);
-			for (int32 Index = 0; Index < SourceAssets.Num(); Index++)
+			const FAssetManagerEditorRegistrySource* RegistrySource = EditorModule->GetCurrentRegistrySource();
+
+			if (RegistrySource->GetOwnedRegistryState())
 			{
-				const FAssetData& SourceAsset = SourceAssets[Index];
-				FAssetTableRow AssetRow;
-				PopulateAssetTableRow(AssetRow, SourceAsset, *AssetTable);
-				AssetTable->AddAsset(AssetRow);
-				AssetToIndexMap.Add(SourceAsset, Index);
+				FARCompiledFilter CompiledFilter;
+				AssetRegistry->CompileFilter(Filter, CompiledFilter);
+				RegistrySource->GetOwnedRegistryState()->GetAssets(CompiledFilter, TSet<FName>(), SourceAssets);
+			}
+			else
+			{
+				AssetRegistry->GetAssets(Filter, SourceAssets);
+			}
+
+			for (int32 SourceAssetIndex = 0; SourceAssetIndex < SourceAssets.Num(); SourceAssetIndex++)
+			{
+				TArray<FAssetData> AssetsInSourcePackage;
+				AssetRegistry->GetAssetsByPackageName(SourceAssets[SourceAssetIndex].PackageName, AssetsInSourcePackage, /*bIncludeOnlyOnDiskAssets*/ true); // Only use on disk assets to avoid creating FAssetData for everything in memory
+
+				for (FAssetData& SourceAsset : AssetsInSourcePackage)
+				{
+					if (AssetToIndexMap.Find(SourceAsset) == nullptr)
+					{
+						FAssetTableRow AssetRow;
+						PopulateAssetTableRow(AssetRow, SourceAsset, *AssetTable);
+						AssetTable->AddAsset(AssetRow);
+						AssetToIndexMap.Add(SourceAsset, AssetTable->GetTotalAssetCount() - 1);
+					}
+				}
 			}
 			AssetTable->SetVisibleAssetCount(AssetTable->GetAssets().Num());
 
 			int32 EntryIndex = AssetTable->GetAssets().Num();
-			const FAssetManagerEditorRegistrySource* RegistrySource = EditorModule->GetCurrentRegistrySource();
 
 			TSet<FName> AlreadyProcessedPackages;
 			// Go over all SourceAssets. We will add more to the end as we find we need to add them.
@@ -878,33 +882,46 @@ void SAssetAuditBrowser::RefreshAssetView()
 
 				// Get the dependencies
 				TArray<FAssetIdentifier> DependencyList;
-//				RegistrySource->GetDependencies(CurrentIdentifier, DependencyList, UE::AssetRegistry::EDependencyCategory::Package);
-				AssetRegistry->GetDependencies(CurrentIdentifier, DependencyList, UE::AssetRegistry::EDependencyCategory::Package);
-
+				RegistrySource->GetDependencies(CurrentIdentifier, DependencyList, UE::AssetRegistry::EDependencyCategory::Package);
+				
 				TArray<int32> IndicesOfDependenciesInRowTableToAddToCurrentSourceAssetRow;
 
 				for (const FAssetIdentifier& Dependency : DependencyList)
 				{
-					TArray<FAssetData> AssetsInDependencyPackage;
-					if (AssetRegistry->GetAssetsByPackageName(Dependency.PackageName, AssetsInDependencyPackage, /*bIncludeOnlyOnDiskAssets*/ true))
+					auto GatherIndices = [&](FAssetData const* DependencyAsset)
 					{
-						for (const FAssetData& DependencyAsset : AssetsInDependencyPackage)
+						if (int32* DependencyIndex = AssetToIndexMap.Find(*DependencyAsset))
 						{
-							if (int32* DependencyIndex = AssetToIndexMap.Find(DependencyAsset))
-							{
-								// Do we already know about the asset? If so, great.
-								IndicesOfDependenciesInRowTableToAddToCurrentSourceAssetRow.Add(*DependencyIndex);
-							}
-							else
-							{
-								// We don't know about this asset yet. Create a new row for it and add it to the source list for further dependency analysis
-								FAssetTableRow NewRow;
-								PopulateAssetTableRow(NewRow, DependencyAsset, *AssetTable);
-								AssetTable->AddAsset(NewRow);
-								IndicesOfDependenciesInRowTableToAddToCurrentSourceAssetRow.Add(SourceAssets.Num());
-								AssetToIndexMap.Add(DependencyAsset, SourceAssets.Num());
-								SourceAssets.Add(DependencyAsset);
-							}
+							// Do we already know about the asset? If so, great.
+							IndicesOfDependenciesInRowTableToAddToCurrentSourceAssetRow.Add(*DependencyIndex);
+						}
+						else
+						{
+							// We don't know about this asset yet. Create a new row for it and add it to the source list for further dependency analysis
+							FAssetTableRow NewRow;
+							PopulateAssetTableRow(NewRow, *DependencyAsset, *AssetTable);
+							AssetTable->AddAsset(NewRow);
+							IndicesOfDependenciesInRowTableToAddToCurrentSourceAssetRow.Add(AssetTable->GetTotalAssetCount() - 1);
+							AssetToIndexMap.Add(*DependencyAsset, AssetTable->GetTotalAssetCount() - 1);
+							SourceAssets.Add(*DependencyAsset);
+						}
+					};
+					
+					if (RegistrySource->GetOwnedRegistryState())
+					{
+						TArrayView<FAssetData const* const> AssetsInDependencyPackage = RegistrySource->GetOwnedRegistryState()->GetAssetsByPackageName(Dependency.PackageName);
+						for (FAssetData const* const DependencyAsset : AssetsInDependencyPackage)
+						{
+							GatherIndices(DependencyAsset);
+						}
+					}
+					else
+					{
+						TArray<FAssetData> AssetsInDependencyPackage;
+						AssetRegistry->GetAssetsByPackageName(Dependency.PackageName, AssetsInDependencyPackage, /*only on disk assets*/true);
+						for (FAssetData& DependencyAsset : AssetsInDependencyPackage)
+						{
+							GatherIndices(&DependencyAsset);
 						}
 					}
 				}
@@ -917,7 +934,26 @@ void SAssetAuditBrowser::RefreshAssetView()
 						int32* RowIndex = AssetToIndexMap.Find(SourceAssetData);
 						if (ensure(RowIndex != nullptr))
 						{
-							AssetTable->GetAsset(*RowIndex)->Dependencies.Append(IndicesOfDependenciesInRowTableToAddToCurrentSourceAssetRow);
+							for (int32 Index : IndicesOfDependenciesInRowTableToAddToCurrentSourceAssetRow)
+							{
+								AssetTable->GetAsset(*RowIndex)->Dependencies.AddUnique(Index);
+							}
+						}
+					}
+
+					for (int32 Index : IndicesOfDependenciesInRowTableToAddToCurrentSourceAssetRow)
+					{
+						// Now for each of those dependencies, add this source asset row as a referencer
+						if (FAssetTableRow* DependentRow = AssetTable->GetAsset(Index))
+						{
+							for (const FAssetData& SourceAssetData : AssetsInSourcePackage)
+							{
+								int32* RowIndex = AssetToIndexMap.Find(SourceAssetData);
+								if (ensure(RowIndex != nullptr))
+								{
+									DependentRow->Referencers.AddUnique(*RowIndex);
+								}
+							}
 						}
 					}
 				}
diff --git a/Engine/Plugins/Editor/AssetManagerEditor/Source/AssetManagerEditor/Private/TreeView/AssetTable.cpp b/Engine/Plugins/Editor/AssetManagerEditor/Source/AssetManagerEditor/Private/TreeView/AssetTable.cpp
--- a/Engine/Plugins/Editor/AssetManagerEditor/Source/AssetManagerEditor/Private/TreeView/AssetTable.cpp
+++ b/Engine/Plugins/Editor/AssetManagerEditor/Source/AssetManagerEditor/Private/TreeView/AssetTable.cpp
@@ -23,16 +23,13 @@ const FName FAssetTableColumns::CountColumnId(TEXT("Count"));
 const FName FAssetTableColumns::NameColumnId(TEXT("Name"));
 const FName FAssetTableColumns::TypeColumnId(TEXT("Type"));
 const FName FAssetTableColumns::PathColumnId(TEXT("Path"));
-//const FName FAssetTableColumns::PackageTypeColumnId(TEXT("PackageType"));
-//const FName FAssetTableColumns::PackageNameColumnId(TEXT("PackageName"));
 const FName FAssetTableColumns::PrimaryTypeColumnId(TEXT("PrimaryType"));
 const FName FAssetTableColumns::PrimaryNameColumnId(TEXT("PrimaryName"));
-//const FName FAssetTableColumns::ManagedDiskSizeColumnId(TEXT("ManagedDiskSize"));
-//const FName FAssetTableColumns::DiskSizeColumnId(TEXT("DiskSize"));
 const FName FAssetTableColumns::StagedCompressedSizeColumnId(TEXT("StagedCompressedSize"));
+const FName FAssetTableColumns::TotalSizeUniqueDependenciesColumnId(TEXT("TotalSizeUniqueDependencies"));
+const FName FAssetTableColumns::TotalSizeOtherDependenciesColumnId(TEXT("TotalSizeOtherDependencies"));
 const FName FAssetTableColumns::TotalUsageCountColumnId(TEXT("TotalUsageCount"));
-//const FName FAssetTableColumns::CookRuleColumnId(TEXT("CookRule"));
-//const FName FAssetTableColumns::ChunksColumnId(TEXT("Chunks"));
+const FName FAssetTableColumns::ChunksColumnId(TEXT("Chunks"));
 const FName FAssetTableColumns::NativeClassColumnId(TEXT("NativeClass"));
 const FName FAssetTableColumns::GameFeaturePluginColumnId(TEXT("GameFeaturePlugin"));
 
@@ -154,6 +151,227 @@ void FAssetTableStringStore::EnumerateStrings(TFunction<void(const FStringView S
 
 FAssetTable::FAssetTable()
 {
+//#define RUN_ASSET_TRAVERSAL_TESTS
+#ifdef RUN_ASSET_TRAVERSAL_TESTS
+	for (int32 i = 0; i < 9; i++)
+	{
+		FAssetTableRow& Asset = Assets.AddDefaulted_GetRef();
+		int32 Id = i;
+		int32 Id2 = 10 * i;
+
+		Asset.Type = StoreStr(FString::Printf(TEXT("Type%02d"), Id % 10));
+		Asset.Name = StoreStr(FString::Printf(TEXT("Name%d"), Id));
+		Asset.Path = StoreStr(FString::Printf(TEXT("A%02d/B%02d/C%02d/D%02d"), Id % 11, Id % 13, Id % 17, Id % 19));
+		Asset.PrimaryType = StoreStr(FString::Printf(TEXT("PT_%02d"), Id2 % 10));
+		Asset.PrimaryName = StoreStr(FString::Printf(TEXT("PN%d"), Id2));
+		Asset.TotalUsageCount = 10;
+		Asset.StagedCompressedSize = 1;
+		Asset.NativeClass = StoreStr(FString::Printf(TEXT("NativeClass%02d"), (Id * Id * Id) % 8));
+		Asset.GameFeaturePlugin = StoreStr(TEXT("MockGFP"));
+	}
+
+	for (int32 TestIndex = 0; TestIndex < 5; TestIndex++)
+	{
+		TArray<TSet<int32>> UniqueDependencies;
+		UniqueDependencies.SetNum(Assets.Num());
+		TArray<TSet<int32>> SharedDependencies;
+		SharedDependencies.SetNum(Assets.Num());
+		// Reset dependencies for next test
+		for (int32 AssetIndex = 0; AssetIndex < Assets.Num(); AssetIndex++)
+		{
+			Assets[AssetIndex].Dependencies.Empty();
+			Assets[AssetIndex].Referencers.Empty();
+		}
+
+		if (TestIndex == 0)
+		{
+			// 0-->1-->2-->3, 4-->3
+			Assets[0].Dependencies.Add(1);
+			Assets[1].Dependencies.Add(2);
+			Assets[2].Dependencies.Add(3);
+			Assets[4].Dependencies.Add(3);
+			Assets[1].Referencers.Add(0);
+			Assets[2].Referencers.Add(1);
+			Assets[3].Referencers.Add(2);
+			Assets[3].Referencers.Add(4);
+
+			// Add the expected results in order 0-->4
+			UniqueDependencies[0] = TSet<int32>{1, 2};
+			UniqueDependencies[1] = TSet<int32>{2};
+			UniqueDependencies[2] = TSet<int32>{};
+			UniqueDependencies[3] = TSet<int32>{};
+			UniqueDependencies[4] = TSet<int32>{};
+
+			SharedDependencies[0] = TSet<int32>{3};
+			SharedDependencies[1] = TSet<int32>{3};
+			SharedDependencies[2] = TSet<int32>{3};
+			SharedDependencies[3] = TSet<int32>{};
+			SharedDependencies[4] = TSet<int32>{3};
+		}
+		else if (TestIndex == 1)
+		{
+			// 0-->1-->2-->3, 4-->0
+			Assets[0].Dependencies.Add(1);
+			Assets[1].Dependencies.Add(2);
+			Assets[2].Dependencies.Add(3);
+			Assets[4].Dependencies.Add(0);
+			Assets[1].Referencers.Add(0);
+			Assets[2].Referencers.Add(1);
+			Assets[3].Referencers.Add(2);
+			Assets[0].Referencers.Add(4);
+
+			// Add the expected results in order 0-->4
+			UniqueDependencies[0] = TSet<int32>{1, 2, 3};
+			UniqueDependencies[1] = TSet<int32>{2, 3};
+			UniqueDependencies[2] = TSet<int32>{3};
+			UniqueDependencies[3] = TSet<int32>{};
+			UniqueDependencies[4] = TSet<int32>{0, 1, 2, 3};
+   
+			SharedDependencies[0] = TSet<int32>{};
+			SharedDependencies[1] = TSet<int32>{};
+			SharedDependencies[2] = TSet<int32>{};
+			SharedDependencies[3] = TSet<int32>{};
+			SharedDependencies[4] = TSet<int32>{};
+		}
+		else if (TestIndex == 2)
+		{
+			// 0-->1-->2-->3-->0, 4-->2
+			Assets[0].Dependencies.Add(1);
+			Assets[1].Dependencies.Add(2);
+			Assets[2].Dependencies.Add(3);
+			Assets[3].Dependencies.Add(0);
+			Assets[4].Dependencies.Add(2);
+			Assets[1].Referencers.Add(0);
+			Assets[2].Referencers.Add(1);
+			Assets[2].Referencers.Add(4);
+			Assets[3].Referencers.Add(2);
+			Assets[0].Referencers.Add(3);
+
+			// Add the expected results in order 0-->4
+			UniqueDependencies[0] = TSet<int32>{1};
+			UniqueDependencies[1] = TSet<int32>{};
+			UniqueDependencies[2] = TSet<int32>{3, 0, 1};
+			UniqueDependencies[3] = TSet<int32>{0, 1};
+
+			// This is an interesting result. When we traverse the graph for Element 4
+			// we find 4-->2-->3-->0-->1-->[terminate loop]. From that point of view,
+			// '2' is not shared because its other referencer is part of 4's dependency chain
+			// therefore everything is treated as a unique dependency of 4.
+			UniqueDependencies[4] = TSet<int32>{2, 3, 0, 1};
+
+			SharedDependencies[0] = TSet<int32>{2, 3};
+			SharedDependencies[1] = TSet<int32>{2, 3, 0};
+			SharedDependencies[2] = TSet<int32>{};
+			SharedDependencies[3] = TSet<int32>{2};
+			SharedDependencies[4] = TSet<int32>{};
+		}
+		else if (TestIndex == 3)
+		{
+			// 0-->1-->2-->3-->0, 4-->2, 5-->2
+			Assets[0].Dependencies.Add(1);
+			Assets[1].Dependencies.Add(2);
+			Assets[2].Dependencies.Add(3);
+			Assets[3].Dependencies.Add(0);
+			Assets[4].Dependencies.Add(2);
+			Assets[5].Dependencies.Add(2);
+			Assets[1].Referencers.Add(0);
+			Assets[2].Referencers.Add(1);
+			Assets[2].Referencers.Add(4);
+			Assets[2].Referencers.Add(5);
+			Assets[3].Referencers.Add(2);
+			Assets[0].Referencers.Add(3);
+
+			// Add the expected results in order 0-->4
+			UniqueDependencies[0] = TSet<int32>{ 1 };
+			UniqueDependencies[1] = TSet<int32>{};
+			UniqueDependencies[2] = TSet<int32>{ 3, 0, 1 };
+			UniqueDependencies[3] = TSet<int32>{ 0, 1 };
+
+			// Unlike the previous test, since 5 also references 2, 
+			// 4 will see itself as having no unique dependencies
+			UniqueDependencies[4] = TSet<int32>{};
+			UniqueDependencies[5] = TSet<int32>{};
+
+			SharedDependencies[0] = TSet<int32>{ 2, 3 };
+			SharedDependencies[1] = TSet<int32>{ 2, 3, 0 };
+			SharedDependencies[2] = TSet<int32>{};
+			SharedDependencies[3] = TSet<int32>{ 2 };
+			SharedDependencies[4] = TSet<int32>{2, 3, 0, 1};
+			SharedDependencies[5] = TSet<int32>{2, 3, 0, 1};
+		}
+		else if (TestIndex == 4)
+		{
+			// 0 --> 1 --> 6
+			// 0 --> 2 --> 6
+			// 0 --> 7
+			// 0 --> 5 --> 8
+			// 3 --> 2 --> 6
+			// 3 --> 5 --> 8
+			// 2 --> 5 --> 8
+			// 3 --> 4
+			// 2 --> 8
+
+			Assets[0].Dependencies.Add(1);
+			Assets[0].Dependencies.Add(2);
+			Assets[1].Dependencies.Add(6);
+			Assets[0].Dependencies.Add(7);
+			Assets[3].Dependencies.Add(2);
+			Assets[2].Dependencies.Add(6);
+			Assets[3].Dependencies.Add(5);
+			Assets[2].Dependencies.Add(5);
+			Assets[0].Dependencies.Add(5);
+			Assets[3].Dependencies.Add(4);
+			Assets[5].Dependencies.Add(8);
+			Assets[2].Dependencies.Add(8);
+
+			Assets[1].Referencers.Add(0);
+			Assets[2].Referencers.Add(0);
+			Assets[6].Referencers.Add(1);
+			Assets[7].Referencers.Add(0);
+			Assets[2].Referencers.Add(3);
+			Assets[6].Referencers.Add(2);
+			Assets[5].Referencers.Add(3);
+			Assets[5].Referencers.Add(2);
+			Assets[5].Referencers.Add(0);
+			Assets[4].Referencers.Add(3);
+			Assets[8].Referencers.Add(2);
+			Assets[8].Referencers.Add(5);
+
+			UniqueDependencies[0] = TSet<int32>{ 1, 7 };
+			UniqueDependencies[1] = TSet<int32>{};
+			UniqueDependencies[2] = TSet<int32>{};
+			UniqueDependencies[3] = TSet<int32>{ 4 };
+			UniqueDependencies[4] = TSet<int32>{};
+			UniqueDependencies[5] = TSet<int32>{};
+			UniqueDependencies[6] = TSet<int32>{};
+			UniqueDependencies[7] = TSet<int32>{};
+			UniqueDependencies[8] = TSet<int32>{};
+
+			SharedDependencies[0] = TSet<int32>{ 2, 6, 5, 8 };
+			SharedDependencies[1] = TSet<int32>{ 6 };
+			SharedDependencies[2] = TSet<int32>{ 5, 6, 8 };
+			SharedDependencies[3] = TSet<int32>{ 2, 5, 8, 6 };
+			SharedDependencies[4] = TSet<int32>{};
+			SharedDependencies[5] = TSet<int32>{ 8 };
+			SharedDependencies[6] = TSet<int32>{};
+			SharedDependencies[7] = TSet<int32>{};
+			SharedDependencies[8] = TSet<int32>{};
+
+		}
+
+		for (int32 AssetIndex = 0; AssetIndex < Assets.Num(); AssetIndex++)
+		{
+			TSet<int32> DiscoveredUniqueDependencies;
+			TSet<int32> DiscoveredSharedDependencies;
+			Assets[AssetIndex].ComputeDependencySizes(this, AssetIndex, &DiscoveredUniqueDependencies, &DiscoveredSharedDependencies);
+			ensureAlways(DiscoveredUniqueDependencies.Num() == UniqueDependencies[AssetIndex].Num()
+				&& DiscoveredUniqueDependencies.Intersect(UniqueDependencies[AssetIndex]).Num() == DiscoveredUniqueDependencies.Num());
+			ensureAlways(DiscoveredSharedDependencies.Num() == SharedDependencies[AssetIndex].Num()
+				&& DiscoveredSharedDependencies.Intersect(SharedDependencies[AssetIndex]).Num() == DiscoveredSharedDependencies.Num());
+		}
+	}
+	Assets.Empty();
+#endif
 #if 0 // debug, mock data
 
 	VisibleAssetCount = 100;
@@ -587,17 +805,16 @@ void FAssetTable::AddDefaultColumns()
 		AddColumn(ColumnRef);
 	}
 	//////////////////////////////////////////////////
-	// Managed Disk Size Column
-#if 0
+	// Staged Compressed Size Column
 	{
-		TSharedRef<FTableColumn> ColumnRef = MakeShared<FTableColumn>(FAssetTableColumns::ManagedDiskSizeColumnId);
+		TSharedRef<FTableColumn> ColumnRef = MakeShared<FTableColumn>(FAssetTableColumns::StagedCompressedSizeColumnId);
 		FTableColumn& Column = *ColumnRef;
 
 		Column.SetIndex(ColumnIndex++);
 
-		Column.SetShortName(LOCTEXT("ManagedDiskSizeColumnName", "Disk Size"));
-		Column.SetTitleName(LOCTEXT("ManagedDiskSizeColumnTitle", "Disk Size"));
-		Column.SetDescription(LOCTEXT("ManagedDiskSizeColumnDesc", "Total disk space used by both this and all managed assets"));
+		Column.SetShortName(LOCTEXT("StagedCompressedSizeColumnName", "Staged Compressed Size"));
+		Column.SetTitleName(LOCTEXT("StagedCompressedSizeColumnTitle", "Staged Compressed Size"));
+		Column.SetDescription(LOCTEXT("StagedCompressedSizeColumnDesc", "Compressed size of iostore chunks for this asset's package. Only visible after staging."));
 
 		Column.SetFlags(ETableColumnFlags::ShouldBeVisible | ETableColumnFlags::CanBeHidden | ETableColumnFlags::CanBeFiltered);
 
@@ -606,7 +823,7 @@ void FAssetTable::AddDefaultColumns()
 
 		Column.SetDataType(ETableCellDataType::Int64);
 
-		class FManagedDiskSizeValueGetter : public FTableCellValueGetter
+		class FStagedCompressedSizeValueGetter : public FTableCellValueGetter
 		{
 		public:
 			virtual const TOptional<FTableCellValue> GetValue(const FTableColumn& Column, const FBaseTreeNode& Node) const override
@@ -623,13 +840,13 @@ void FAssetTable::AddDefaultColumns()
 				{
 					const FAssetTreeNode& TreeNode = static_cast<const FAssetTreeNode&>(Node);
 					const FAssetTableRow& Asset = TreeNode.GetAssetChecked();
-					return FTableCellValue(static_cast<int64>(Asset.GetManagedDiskSize()));
+					return FTableCellValue(static_cast<int64>(Asset.GetStagedCompressedSize()));
 				}
 
 				return TOptional<FTableCellValue>();
 			}
 		};
-		TSharedRef<ITableCellValueGetter> Getter = MakeShared<FManagedDiskSizeValueGetter>();
+		TSharedRef<ITableCellValueGetter> Getter = MakeShared<FStagedCompressedSizeValueGetter>();
 		Column.SetValueGetter(Getter);
 
 		TSharedRef<ITableCellValueFormatter> Formatter = MakeShared<FInt64ValueFormatterAsMemory>();
@@ -646,28 +863,26 @@ void FAssetTable::AddDefaultColumns()
 
 		AddColumn(ColumnRef);
 	}
-#endif
 	//////////////////////////////////////////////////
-	// Exclusive Disk Size Column
-#if 0
+	// Total Size of Unique Dependencies
 	{
-		TSharedRef<FTableColumn> ColumnRef = MakeShared<FTableColumn>(FAssetTableColumns::DiskSizeColumnId);
+		TSharedRef<FTableColumn> ColumnRef = MakeShared<FTableColumn>(FAssetTableColumns::TotalSizeUniqueDependenciesColumnId);
 		FTableColumn& Column = *ColumnRef;
 
 		Column.SetIndex(ColumnIndex++);
 
-		Column.SetShortName(LOCTEXT("DiskSizeColumnName", "Exclusive Disk Size"));
-		Column.SetTitleName(LOCTEXT("DiskSizeColumnTitle", "Exclusive Disk Size"));
-		Column.SetDescription(LOCTEXT("DiskSizeColumnDesc", "Size of saved file(s) on disk for this asset's package. If Asset Registry Writeback is enabled and a platform is selected, this will be the uncompressed size of the iostore chunks for the asset's package."));
+		Column.SetShortName(LOCTEXT("TotalSizeUniqueDependenciesColumnName", "Total Unique Dependency Size"));
+		Column.SetTitleName(LOCTEXT("TotalSizeUniqueDependenciesColumnTitle", "Total Unique Dependency Size"));
+		Column.SetDescription(LOCTEXT("TotalSizeUniqueDependenciesColumnIdDesc", "Sum of the staged compressed sizes of all dependencies of this asset, counted only once"));
 
 		Column.SetFlags(ETableColumnFlags::ShouldBeVisible | ETableColumnFlags::CanBeHidden | ETableColumnFlags::CanBeFiltered);
 
 		Column.SetHorizontalAlignment(HAlign_Right);
-		Column.SetInitialWidth(100.0f);
+		Column.SetInitialWidth(120.0f);
 
 		Column.SetDataType(ETableCellDataType::Int64);
 
-		class FDiskSizeValueGetter : public FTableCellValueGetter
+		class FTotalSizeUniqueDependenciesValueGetter : public FTableCellValueGetter
 		{
 		public:
 			virtual const TOptional<FTableCellValue> GetValue(const FTableColumn& Column, const FBaseTreeNode& Node) const override
@@ -684,13 +899,13 @@ void FAssetTable::AddDefaultColumns()
 				{
 					const FAssetTreeNode& TreeNode = static_cast<const FAssetTreeNode&>(Node);
 					const FAssetTableRow& Asset = TreeNode.GetAssetChecked();
-					return FTableCellValue(static_cast<int64>(Asset.GetDiskSize()));
+					return FTableCellValue(static_cast<int64>(Asset.GetOrComputeTotalSizeUniqueDependencies(&TreeNode.GetAssetTableChecked(), TreeNode.GetRowIndex())));
 				}
 
 				return TOptional<FTableCellValue>();
 			}
 		};
-		TSharedRef<ITableCellValueGetter> Getter = MakeShared<FDiskSizeValueGetter>();
+		TSharedRef<ITableCellValueGetter> Getter = MakeShared<FTotalSizeUniqueDependenciesValueGetter>();
 		Column.SetValueGetter(Getter);
 
 		TSharedRef<ITableCellValueFormatter> Formatter = MakeShared<FInt64ValueFormatterAsMemory>();
@@ -703,31 +918,30 @@ void FAssetTable::AddDefaultColumns()
 		//TSharedRef<IFilterValueConverter> Converter = MakeShared<FMemoryFilterValueConverter>();
 		//Column.SetValueConverter(Converter);
 
-		Column.SetAggregation(ETableColumnAggregation::Sum);
+		Column.SetAggregation(ETableColumnAggregation::None);
 
 		AddColumn(ColumnRef);
 	}
-#endif
 	//////////////////////////////////////////////////
-	// Staged Compressed Size Column
+	// Total Size of Other Dependencies
 	{
-		TSharedRef<FTableColumn> ColumnRef = MakeShared<FTableColumn>(FAssetTableColumns::StagedCompressedSizeColumnId);
+		TSharedRef<FTableColumn> ColumnRef = MakeShared<FTableColumn>(FAssetTableColumns::TotalSizeOtherDependenciesColumnId);
 		FTableColumn& Column = *ColumnRef;
 
 		Column.SetIndex(ColumnIndex++);
 
-		Column.SetShortName(LOCTEXT("StagedCompressedSizeColumnName", "Staged Compressed Size"));
-		Column.SetTitleName(LOCTEXT("StagedCompressedSizeColumnTitle", "Staged Compressed Size"));
-		Column.SetDescription(LOCTEXT("StagedCompressedSizeColumnDesc", "Compressed size of iostore chunks for this asset's package. Only visible after staging."));
+		Column.SetShortName(LOCTEXT("TotalSizeOtherDependenciesColumnName", "Total Other Dependency Size"));
+		Column.SetTitleName(LOCTEXT("TotalSizeOtherDependenciesColumnTitle", "Total Other Dependency Size"));
+		Column.SetDescription(LOCTEXT("TotalSizeOtherDependenciesColumnIdDesc", "Sum of the staged compressed sizes of all dependencies of this asset which are shared by other assets directly or indirectly, counted only once"));
 
 		Column.SetFlags(ETableColumnFlags::ShouldBeVisible | ETableColumnFlags::CanBeHidden | ETableColumnFlags::CanBeFiltered);
 
 		Column.SetHorizontalAlignment(HAlign_Right);
-		Column.SetInitialWidth(100.0f);
+		Column.SetInitialWidth(120.0f);
 
 		Column.SetDataType(ETableCellDataType::Int64);
 
-		class FStagedCompressedSizeValueGetter : public FTableCellValueGetter
+		class FTotalSizeUniqueDependenciesValueGetter : public FTableCellValueGetter
 		{
 		public:
 			virtual const TOptional<FTableCellValue> GetValue(const FTableColumn& Column, const FBaseTreeNode& Node) const override
@@ -744,13 +958,13 @@ void FAssetTable::AddDefaultColumns()
 				{
 					const FAssetTreeNode& TreeNode = static_cast<const FAssetTreeNode&>(Node);
 					const FAssetTableRow& Asset = TreeNode.GetAssetChecked();
-					return FTableCellValue(static_cast<int64>(Asset.GetStagedCompressedSize()));
+					return FTableCellValue(static_cast<int64>(Asset.GetOrComputeTotalSizeOtherDependencies(&TreeNode.GetAssetTableChecked(), TreeNode.GetRowIndex())));
 				}
 
 				return TOptional<FTableCellValue>();
 			}
 		};
-		TSharedRef<ITableCellValueGetter> Getter = MakeShared<FStagedCompressedSizeValueGetter>();
+		TSharedRef<ITableCellValueGetter> Getter = MakeShared<FTotalSizeUniqueDependenciesValueGetter>();
 		Column.SetValueGetter(Getter);
 
 		TSharedRef<ITableCellValueFormatter> Formatter = MakeShared<FInt64ValueFormatterAsMemory>();
@@ -763,7 +977,7 @@ void FAssetTable::AddDefaultColumns()
 		//TSharedRef<IFilterValueConverter> Converter = MakeShared<FMemoryFilterValueConverter>();
 		//Column.SetValueConverter(Converter);
 
-		Column.SetAggregation(ETableColumnAggregation::Sum);
+		Column.SetAggregation(ETableColumnAggregation::None);
 
 		AddColumn(ColumnRef);
 	}
@@ -823,66 +1037,8 @@ void FAssetTable::AddDefaultColumns()
 
 		AddColumn(ColumnRef);
 	}
-	//////////////////////////////////////////////////
-	// Cook Rule Column
-#if 0
-	{
-		TSharedRef<FTableColumn> ColumnRef = MakeShared<FTableColumn>(FAssetTableColumns::CookRuleColumnId);
-		FTableColumn& Column = *ColumnRef;
-
-		Column.SetIndex(ColumnIndex++);
-
-		Column.SetShortName(LOCTEXT("CookRuleColumnName", "Cook Rule"));
-		Column.SetTitleName(LOCTEXT("CookRuleColumnTitle", "Cook Rule"));
-		Column.SetDescription(LOCTEXT("CookRuleColumnDesc", "Whether this asset will be cooked or not"));
-
-		Column.SetFlags(ETableColumnFlags::ShouldBeVisible | ETableColumnFlags::CanBeHidden | ETableColumnFlags::CanBeFiltered);
-
-		Column.SetHorizontalAlignment(HAlign_Left);
-		Column.SetInitialWidth(120.0f);
-
-		Column.SetDataType(ETableCellDataType::CString);
-
-		class FCookRuleValueGetter : public FTableCellValueGetter
-		{
-		public:
-			virtual const TOptional<FTableCellValue> GetValue(const FTableColumn& Column, const FBaseTreeNode& Node) const
-			{
-				if (Node.IsGroup())
-				{
-					const FTableTreeNode& NodePtr = static_cast<const FTableTreeNode&>(Node);
-					if (NodePtr.HasAggregatedValue(Column.GetId()))
-					{
-						return NodePtr.GetAggregatedValue(Column.GetId());
-					}
-				}
-				else //if (Node->Is<FAssetTreeNode>())
-				{
-					const FAssetTreeNode& TreeNode = static_cast<const FAssetTreeNode&>(Node);
-					const FAssetTableRow& Asset = TreeNode.GetAssetChecked();
-					return FTableCellValue(Asset.GetChunks());
-				}
-
-				return TOptional<FTableCellValue>();
-			}
-		};
-		TSharedRef<ITableCellValueGetter> Getter = MakeShared<FCookRuleValueGetter>();
-		Column.SetValueGetter(Getter);
-
-		TSharedRef<ITableCellValueFormatter> Formatter = MakeShared<FCStringValueFormatterAsText>();
-		Column.SetValueFormatter(Formatter);
-
-		TSharedRef<ITableCellValueSorter> Sorter = MakeShared<FSorterByCStringValue>(ColumnRef);
-		Column.SetValueSorter(Sorter);
-
-		Column.SetAggregation(ETableColumnAggregation::SameValue);
-
-		AddColumn(ColumnRef);
-	}
-#endif
 	//////////////////////////////////////////////////
 	// Chunks Column
-#if 0
 	{
 		TSharedRef<FTableColumn> ColumnRef = MakeShared<FTableColumn>(FAssetTableColumns::ChunksColumnId);
 		FTableColumn& Column = *ColumnRef;
@@ -936,7 +1092,6 @@ void FAssetTable::AddDefaultColumns()
 
 		AddColumn(ColumnRef);
 	}
-#endif
 	//////////////////////////////////////////////////
 	// Native Class Column
 	{
@@ -1050,6 +1205,184 @@ void FAssetTable::AddDefaultColumns()
 	//////////////////////////////////////////////////
 }
 
+// Iteratively refines the dependency set into unique and shared sets. PreviouslyVisitedIndices is split in each pass, moving some elements into 
+// OutExcludedIndices (preserving its original contents) and putting those that are still potentially uniquely owned by the ThisIndex asset into 
+// OutIncrementallyRefinedUniqueIndices. Returns 'true' if another pass is required (i.e., if work was done).
+/*static*/ bool FAssetTableRow::RefineDependencies(TSet<int32> PreviouslyVisitedIndices, FAssetTable* OwningTable, int32 ThisIndex, TSet<int32>* OutIncrementallyRefinedUniqueIndices, TSet<int32>* OutExcludedIndices)
+{
+	const FString& ThisGFP = OwningTable->GetAsset(ThisIndex)->GameFeaturePlugin;
+
+	ensure(OutIncrementallyRefinedUniqueIndices != nullptr);
+	ensure(OutExcludedIndices != nullptr);
+
+	// "visit" ThisIndex to seed the exploration
+	TArray<int32> IndicesToVisit = OwningTable->GetAsset(ThisIndex)->GetDependencies();
+
+	while (IndicesToVisit.Num() > 0)
+	{
+		int32 CurrentIndex = IndicesToVisit.Pop(false);
+
+		FAssetTableRow* Row = OwningTable->GetAsset(CurrentIndex);
+		if (Row->GameFeaturePlugin != ThisGFP)
+		{
+			// Don't traverse outside this plugin
+			continue;
+		}
+
+		bool ShouldIncludeInTotal = true;
+		for (int32 ReferencerIndex : Row->Referencers)
+		{
+			if (PreviouslyVisitedIndices.Contains(ReferencerIndex) == false && ThisIndex != ReferencerIndex)
+			{
+				ShouldIncludeInTotal = false;
+				break;
+			}
+		}
+		if (ShouldIncludeInTotal == false)
+		{
+			OutExcludedIndices->Add(CurrentIndex);
+			continue;
+		}
+		OutIncrementallyRefinedUniqueIndices->Add(CurrentIndex);
+
+		for (int32 ChildIndex : OwningTable->GetAsset(CurrentIndex)->GetDependencies())
+		{
+			// Don't revisit nodes we've already visited and don't re-add ThisIndex to avoid loops (and to avoid counting ourself)
+			if (!OutIncrementallyRefinedUniqueIndices->Contains(ChildIndex) && ThisIndex != ChildIndex)
+			{
+				IndicesToVisit.AddUnique(ChildIndex);
+			}
+		}
+	}
+
+	return OutIncrementallyRefinedUniqueIndices->Num() != PreviouslyVisitedIndices.Num();
+}
+
+void FAssetTableRow::ComputeDependencySizes(FAssetTable* OwningTable, int32 ThisIndex, TSet<int32>* OutUniqueDependencies/* = nullptr*/, TSet<int32>* OutSharedDependencies/* = nullptr*/) const
+{
+	if (OutUniqueDependencies != nullptr && OutSharedDependencies != nullptr)
+	{
+		TotalSizeUniqueDependencies = -1;
+	}
+	if (TotalSizeUniqueDependencies == -1)
+	{
+		TArray<int32> IndicesToVisit;
+		TSet<int32> VisitedIndices;
+
+		TotalSizeUniqueDependencies = 0;
+		TotalSizeOtherDependencies = 0;
+		
+		// Break any loops in the dependency graph
+		VisitedIndices.Add(ThisIndex);
+
+		const FString& ThisGFP = OwningTable->GetAsset(ThisIndex)->GameFeaturePlugin;
+
+		// Don't include this asset itself in the total, just its children
+		for (int32 ChildIndex : Dependencies)
+		{
+			IndicesToVisit.Add(ChildIndex);
+		}
+
+		while (IndicesToVisit.Num() > 0)
+		{
+			int32 CurrentIndex = IndicesToVisit.Pop(false);
+
+			FAssetTableRow* Row = OwningTable->GetAsset(CurrentIndex);
+			if (Row->GameFeaturePlugin != ThisGFP)
+			{
+				// Don't traverse outside this plugin
+				continue;
+			}
+			VisitedIndices.Add(CurrentIndex);
+
+			for (int32 ChildIndex : OwningTable->GetAsset(CurrentIndex)->GetDependencies())
+			{
+				if (!VisitedIndices.Contains(ChildIndex))
+				{
+					IndicesToVisit.AddUnique(ChildIndex);
+				}
+			}
+		}
+
+		// Iteratively separate the graph of "all things referenced by ThisIndex, directly or indirectly"
+		// into "UniqueDependencies -- things referenced ONLY by ThisIndex and by other things themselves referenced ONLY by ThisIndex"
+		// and "OtherDependencies" -- things removed from the list of "all things referenced by ThisIndex" in order to identify UniqueDependencies
+		TSet<int32> UniqueDependencies;
+		TSet<int32> OtherDependencies;
+		while (RefineDependencies(VisitedIndices, OwningTable, ThisIndex, &UniqueDependencies, &OtherDependencies))
+		{
+			VisitedIndices = UniqueDependencies;
+			UniqueDependencies.Empty();
+		}
+
+		for (int32 Index : UniqueDependencies)
+		{
+			int64 DependencySize = OwningTable->GetAsset(Index)->GetStagedCompressedSize();
+			TotalSizeUniqueDependencies += DependencySize;
+		}
+		if (OutUniqueDependencies != nullptr)
+		{
+			*OutUniqueDependencies = UniqueDependencies;
+		}
+
+		// Now explore all the dependencies of OtherDependencies and gather up their sizes
+		// This is necessary because the process calling RefineDependencies doesn't produce a complete list of other dependencies,
+		// just a partial set. By exploring the dependencies of that partial set we can find all the things that were 
+		// referenced by ThisIndex but also by some other asset outside the subgraph defined by ThisIndex and its UniqueDependencies
+		VisitedIndices.Empty();
+		IndicesToVisit.Empty();
+		for (int32 Index : OtherDependencies)
+		{
+			IndicesToVisit.Add(Index);
+		}
+
+		while (IndicesToVisit.Num() > 0)
+		{
+			int32 CurrentIndex = IndicesToVisit.Pop(false);
+			FAssetTableRow* Row = OwningTable->GetAsset(CurrentIndex);
+			VisitedIndices.Add(CurrentIndex);
+			if (Row->GameFeaturePlugin != ThisGFP)
+			{
+				// Don't traverse outside this plugin
+				continue;
+			}
+			
+			int64 DependencySize = OwningTable->GetAsset(CurrentIndex)->GetStagedCompressedSize();
+			TotalSizeOtherDependencies += DependencySize;
+			if (OutSharedDependencies != nullptr)
+			{
+				OutSharedDependencies->Add(CurrentIndex);
+			}
+
+			for (int32 ChildIndex : OwningTable->GetAsset(CurrentIndex)->GetDependencies())
+			{
+				if (!VisitedIndices.Contains(ChildIndex) && ThisIndex != ChildIndex)
+				{
+					IndicesToVisit.AddUnique(ChildIndex);
+				}
+			}
+		}
+	}
+}
+
+int64 FAssetTableRow::GetOrComputeTotalSizeUniqueDependencies(FAssetTable* OwningTable, int32 ThisIndex) const
+{
+	if (TotalSizeUniqueDependencies == -1)
+	{
+		ComputeDependencySizes(OwningTable, ThisIndex);
+	}
+	return TotalSizeUniqueDependencies;
+}
+
+int64 FAssetTableRow::GetOrComputeTotalSizeOtherDependencies(FAssetTable* OwningTable, int32 ThisIndex) const
+{
+	if (TotalSizeOtherDependencies == -1)
+	{
+		ComputeDependencySizes(OwningTable, ThisIndex);
+	}
+	return TotalSizeOtherDependencies;
+}
+
 ////////////////////////////////////////////////////////////////////////////////////////////////////
 
 #undef LOCTEXT_NAMESPACE
diff --git a/Engine/Plugins/Editor/AssetManagerEditor/Source/AssetManagerEditor/Private/TreeView/AssetTable.h b/Engine/Plugins/Editor/AssetManagerEditor/Source/AssetManagerEditor/Private/TreeView/AssetTable.h
--- a/Engine/Plugins/Editor/AssetManagerEditor/Source/AssetManagerEditor/Private/TreeView/AssetTable.h
+++ b/Engine/Plugins/Editor/AssetManagerEditor/Source/AssetManagerEditor/Private/TreeView/AssetTable.h
@@ -112,16 +112,13 @@ struct FAssetTableColumns
 	static const FName TypeColumnId;
 	static const FName NameColumnId;
 	static const FName PathColumnId;
-	//static const FName PackageNameColumnId;
-	//static const FName PackageTypeColumnId;
 	static const FName PrimaryTypeColumnId;
 	static const FName PrimaryNameColumnId;
-	//static const FName ManagedDiskSizeColumnId;
-	//static const FName DiskSizeColumnId;
 	static const FName StagedCompressedSizeColumnId;
+	static const FName TotalSizeUniqueDependenciesColumnId;
+	static const FName TotalSizeOtherDependenciesColumnId;
 	static const FName TotalUsageCountColumnId;
-	//static const FName CookRuleColumnId;
-	//static const FName ChunksColumnId;
+	static const FName ChunksColumnId;
 	static const FName NativeClassColumnId;
 	static const FName GameFeaturePluginColumnId;
 };
@@ -142,16 +139,13 @@ public:
 	const TCHAR* GetType() const { return Type; }
 	const TCHAR* GetName() const { return Name; }
 	const TCHAR* GetPath() const { return Path; }
-	//const TCHAR* GetPackageName() const { return PackageName; }
-	//const TCHAR* GetPackagePath() const { return PackagePath; }
 	const TCHAR* GetPrimaryType() const { return PrimaryType; }
 	const TCHAR* GetPrimaryName() const { return PrimaryName; }
-	//int64 GetManagedDiskSize() const { return ManagedDiskSize; }
-	//int64 GetDiskSize() const { return DiskSize; }
 	int64 GetStagedCompressedSize() const { return StagedCompressedSize; }
+	int64 GetOrComputeTotalSizeUniqueDependencies(FAssetTable* OwningTable, int32 ThisIndex) const;
+	int64 GetOrComputeTotalSizeOtherDependencies(FAssetTable* OwningTable, int32 ThisIndex) const; 
 	int64 GetTotalUsageCount() const { return TotalUsageCount; }
-	//const TCHAR* GetCookRule() const { return CookRule; }
-	//const TCHAR* GetChunks() const { return Chunks; }
+	const TCHAR* GetChunks() const { return Chunks; }
 	const TCHAR* GetNativeClass() const { return NativeClass; }
 	const TCHAR* GetGameFeaturePlugin() const { return GameFeaturePlugin; }
 	int32 GetNumDependencies() const { return Dependencies.Num(); }
@@ -160,22 +154,24 @@ public:
 	FLinearColor GetColor() const { return Color; }
 
 private:
+	static bool RefineDependencies(TSet<int32> PreviouslyVisitedIndices, FAssetTable* OwningTable, int32 ThisIndex, TSet<int32>* OutIncrementallyRefinedUniqueIndices, TSet<int32>* OutExcludedIndices);
+	// Providing OutUniqueDependencies and OutSharedDependencies will force a full recomputation. It's intended for use in unit testing/validation.
+	void ComputeDependencySizes(FAssetTable* OwningTable, int32 ThisIndex, TSet<int32>* OutUniqueDependencies = nullptr, TSet<int32>* OutSharedDependencies = nullptr) const;
+
 	const TCHAR* Type = nullptr;
 	const TCHAR* Name = nullptr;
 	const TCHAR* Path = nullptr;
-	//const TCHAR* PackageName = nullptr;
-	//const TCHAR* PackagePath = nullptr;
 	const TCHAR* PrimaryType = nullptr;
 	const TCHAR* PrimaryName = nullptr;
-	//int64 ManagedDiskSize = 0;
-	//int64 DiskSize = 0;
 	int64 StagedCompressedSize = 0;
+	mutable int64 TotalSizeUniqueDependencies = -1; // Lazily calculated
+	mutable int64 TotalSizeOtherDependencies = -1; // Lazily calculated
 	int64 TotalUsageCount = 0;
-	//const TCHAR* CookRule = nullptr;
-	//const TCHAR* Chunks = nullptr;
+	const TCHAR* Chunks = nullptr;
 	const TCHAR* NativeClass = nullptr;
 	const TCHAR* GameFeaturePlugin = nullptr;
 	TArray<int32> Dependencies;
+	TArray<int32> Referencers;
 	FLinearColor Color;
 };
 
diff --git a/Engine/Plugins/Editor/AssetManagerEditor/Source/AssetManagerEditor/Private/TreeView/SAssetTableTreeView.cpp b/Engine/Plugins/Editor/AssetManagerEditor/Source/AssetManagerEditor/Private/TreeView/SAssetTableTreeView.cpp
--- a/Engine/Plugins/Editor/AssetManagerEditor/Source/AssetManagerEditor/Private/TreeView/SAssetTableTreeView.cpp
+++ b/Engine/Plugins/Editor/AssetManagerEditor/Source/AssetManagerEditor/Private/TreeView/SAssetTableTreeView.cpp
@@ -258,11 +258,8 @@ void SAssetTableTreeView::InitAvailableViewPresets()
 			InOutConfigSet.Add({ FAssetTableColumns::PathColumnId,                 true, 400.0f });
 			InOutConfigSet.Add({ FAssetTableColumns::PrimaryTypeColumnId,          true, 200.0f });
 			InOutConfigSet.Add({ FAssetTableColumns::PrimaryNameColumnId,          true, 200.0f });
-			//InOutConfigSet.Add({ FAssetTableColumns::ManagedDiskSizeColumnId,      true, 100.0f });
-			//InOutConfigSet.Add({ FAssetTableColumns::DiskSizeColumnId,             true, 100.0f });
 			InOutConfigSet.Add({ FAssetTableColumns::StagedCompressedSizeColumnId, true, 100.0f });
 			InOutConfigSet.Add({ FAssetTableColumns::TotalUsageCountColumnId,      true, 100.0f });
-			//InOutConfigSet.Add({ FAssetTableColumns::CookRuleColumnId,             true, 200.0f });
 			//InOutConfigSet.Add({ FAssetTableColumns::ChunksColumnId,               true, 200.0f });
 			InOutConfigSet.Add({ FAssetTableColumns::NativeClassColumnId,          true, 200.0f });
 			InOutConfigSet.Add({ FAssetTableColumns::GameFeaturePluginColumnId,    true, 200.0f });
@@ -270,6 +267,85 @@ void SAssetTableTreeView::InitAvailableViewPresets()
 	};
 	AvailableViewPresets.Add(MakeShared<FDefaultViewPreset>());
 
+	//////////////////////////////////////////////////
+	// GameFeaturePlugin, PrimaryType, Dependency View
+	
+	class FGameFeaturePluginPrimaryTypeDependencyView : public UE::Insights::ITableTreeViewPreset
+	{
+	public:
+		virtual FText GetName() const override
+		{
+			return LOCTEXT("GFPTypeDepView_PresetName", "DependencyAnalysis");
+		}
+
+		virtual FText GetToolTip() const override
+		{
+			return LOCTEXT("GFPTypeDepView_PresentName", "Dependency Analysis View\nconfigure the tree view to show a breakdown of assets by Game Feature Plugin, Primary Asset Type, and Dependencies.");
+		}
+		virtual FName GetSortColumn() const override
+		{
+			return UE::Insights::FTable::GetHierarchyColumnId();
+		}
+		virtual EColumnSortMode::Type GetSortMode() const override
+		{
+			return EColumnSortMode::Type::Ascending;
+		}
+		virtual void SetCurrentGroupings(const TArray<TSharedPtr<UE::Insights::FTreeNodeGrouping>>& InAvailableGroupings, TArray<TSharedPtr<UE::Insights::FTreeNodeGrouping>>& InOutCurrentGroupings) const override
+		{
+			InOutCurrentGroupings.Reset();
+
+			check(InAvailableGroupings[0]->Is<UE::Insights::FTreeNodeGroupingFlat>());
+			InOutCurrentGroupings.Add(InAvailableGroupings[0]);
+
+			const TSharedPtr<UE::Insights::FTreeNodeGrouping>* GameFeaturePluginGrouping = InAvailableGroupings.FindByPredicate(
+				[](TSharedPtr<UE::Insights::FTreeNodeGrouping>& Grouping)
+				{
+					return Grouping->Is<UE::Insights::FTreeNodeGroupingByUniqueValueCString>() &&
+						Grouping->As<UE::Insights::FTreeNodeGroupingByUniqueValueCString>().GetColumnId() == FAssetTableColumns::GameFeaturePluginColumnId;
+				});
+			if (GameFeaturePluginGrouping)
+			{
+				InOutCurrentGroupings.Add(*GameFeaturePluginGrouping);
+			}
+
+			const TSharedPtr<UE::Insights::FTreeNodeGrouping>* PrimaryTypeGrouping = InAvailableGroupings.FindByPredicate(
+				[](TSharedPtr<UE::Insights::FTreeNodeGrouping>& Grouping)
+				{
+					return Grouping->Is<UE::Insights::FTreeNodeGroupingByUniqueValueCString>() &&
+						Grouping->As<UE::Insights::FTreeNodeGroupingByUniqueValueCString>().GetColumnId() == FAssetTableColumns::PrimaryTypeColumnId;
+				});
+			if (PrimaryTypeGrouping)
+			{
+				InOutCurrentGroupings.Add(*PrimaryTypeGrouping);
+			}
+
+			const TSharedPtr<UE::Insights::FTreeNodeGrouping>* DependencyGrouping = InAvailableGroupings.FindByPredicate(
+				[](TSharedPtr<UE::Insights::FTreeNodeGrouping>& Grouping)
+				{
+					return Grouping->Is<FAssetDependencyGrouping>();
+				});
+			if (DependencyGrouping)
+			{
+				InOutCurrentGroupings.Add(*DependencyGrouping);
+			}
+		}
+		virtual void GetColumnConfigSet(TArray<UE::Insights::FTableColumnConfig>& InOutConfigSet) const override
+		{
+			InOutConfigSet.Add({ UE::Insights::FTable::GetHierarchyColumnId(),     true, 400.0f });
+			InOutConfigSet.Add({ FAssetTableColumns::CountColumnId,                true, 100.0f });
+			InOutConfigSet.Add({ FAssetTableColumns::NameColumnId,                 true, 200.0f });
+			InOutConfigSet.Add({ FAssetTableColumns::PrimaryTypeColumnId,          true, 200.0f });
+			InOutConfigSet.Add({ FAssetTableColumns::StagedCompressedSizeColumnId, true, 100.0f });
+			InOutConfigSet.Add({ FAssetTableColumns::TotalSizeUniqueDependenciesColumnId, true, 180.0f });
+			InOutConfigSet.Add({ FAssetTableColumns::TotalSizeOtherDependenciesColumnId, true, 180.0f });
+			InOutConfigSet.Add({ FAssetTableColumns::NativeClassColumnId,          true, 100.0f });
+			InOutConfigSet.Add({ FAssetTableColumns::GameFeaturePluginColumnId,    true, 200.0f });
+			InOutConfigSet.Add({ FAssetTableColumns::PrimaryNameColumnId,          true, 200.0f });
+			InOutConfigSet.Add({ FAssetTableColumns::TypeColumnId,                 true, 200.0f });
+		}
+	};
+	AvailableViewPresets.Add(MakeShared<FGameFeaturePluginPrimaryTypeDependencyView>());
+
 	//////////////////////////////////////////////////
 	// Path Breakdown View
 
@@ -394,6 +470,7 @@ void SAssetTableTreeView::InitAvailableViewPresets()
 			//InOutConfigSet.Add({ FAssetTableColumns::ManagedDiskSizeColumnId,      true, 100.0f });
 			//InOutConfigSet.Add({ FAssetTableColumns::DiskSizeColumnId,             true, 100.0f });
 			InOutConfigSet.Add({ FAssetTableColumns::StagedCompressedSizeColumnId, true, 100.0f });
+			InOutConfigSet.Add({ FAssetTableColumns::TotalSizeUniqueDependenciesColumnId, true, 100.0f });
 			InOutConfigSet.Add({ FAssetTableColumns::TotalUsageCountColumnId,      true, 100.0f });
 			//InOutConfigSet.Add({ FAssetTableColumns::CookRuleColumnId,             true, 200.0f });
 			//InOutConfigSet.Add({ FAssetTableColumns::ChunksColumnId,               true, 200.0f });
