Static analysis fixes

diff --git a/Engine/Plugins/Editor/AssetManagerEditor/Source/AssetManagerEditor/Private/TreeView/AssetTable.cpp b/Engine/Plugins/Editor/AssetManagerEditor/Source/AssetManagerEditor/Private/TreeView/AssetTable.cpp
--- a/Engine/Plugins/Editor/AssetManagerEditor/Source/AssetManagerEditor/Private/TreeView/AssetTable.cpp
+++ b/Engine/Plugins/Editor/AssetManagerEditor/Source/AssetManagerEditor/Private/TreeView/AssetTable.cpp
@@ -1210,52 +1210,56 @@ void FAssetTable::AddDefaultColumns()
 // OutIncrementallyRefinedUniqueIndices. Returns 'true' if another pass is required (i.e., if work was done).
 /*static*/ bool FAssetTableRow::RefineDependencies(TSet<int32> PreviouslyVisitedIndices, FAssetTable* OwningTable, int32 ThisIndex, TSet<int32>* OutIncrementallyRefinedUniqueIndices, TSet<int32>* OutExcludedIndices)
 {
-	const FString& ThisGFP = OwningTable->GetAsset(ThisIndex)->GameFeaturePlugin;
-
-	ensure(OutIncrementallyRefinedUniqueIndices != nullptr);
-	ensure(OutExcludedIndices != nullptr);
-
-	// "visit" ThisIndex to seed the exploration
-	TArray<int32> IndicesToVisit = OwningTable->GetAsset(ThisIndex)->GetDependencies();
-
-	while (IndicesToVisit.Num() > 0)
+	if (ensure(OutIncrementallyRefinedUniqueIndices != nullptr) && ensure(OutExcludedIndices != nullptr))
 	{
-		int32 CurrentIndex = IndicesToVisit.Pop(false);
+		const FString& ThisGFP = OwningTable->GetAsset(ThisIndex)->GameFeaturePlugin;
 
-		FAssetTableRow* Row = OwningTable->GetAsset(CurrentIndex);
-		if (Row->GameFeaturePlugin != ThisGFP)
-		{
-			// Don't traverse outside this plugin
-			continue;
-		}
+		// "visit" ThisIndex to seed the exploration
+		TArray<int32> IndicesToVisit = OwningTable->GetAsset(ThisIndex)->GetDependencies();
 
-		bool ShouldIncludeInTotal = true;
-		for (int32 ReferencerIndex : Row->Referencers)
+		while (IndicesToVisit.Num() > 0)
 		{
-			if (PreviouslyVisitedIndices.Contains(ReferencerIndex) == false && ThisIndex != ReferencerIndex)
+			int32 CurrentIndex = IndicesToVisit.Pop(false);
+
+			FAssetTableRow* Row = OwningTable->GetAsset(CurrentIndex);
+			if (Row->GameFeaturePlugin != ThisGFP)
 			{
-				ShouldIncludeInTotal = false;
-				break;
+				// Don't traverse outside this plugin
+				continue;
 			}
-		}
-		if (ShouldIncludeInTotal == false)
-		{
-			OutExcludedIndices->Add(CurrentIndex);
-			continue;
-		}
-		OutIncrementallyRefinedUniqueIndices->Add(CurrentIndex);
 
-		for (int32 ChildIndex : OwningTable->GetAsset(CurrentIndex)->GetDependencies())
-		{
-			// Don't revisit nodes we've already visited and don't re-add ThisIndex to avoid loops (and to avoid counting ourself)
-			if (!OutIncrementallyRefinedUniqueIndices->Contains(ChildIndex) && ThisIndex != ChildIndex)
+			bool ShouldIncludeInTotal = true;
+			for (int32 ReferencerIndex : Row->Referencers)
 			{
-				IndicesToVisit.AddUnique(ChildIndex);
+				if (PreviouslyVisitedIndices.Contains(ReferencerIndex) == false && ThisIndex != ReferencerIndex)
+				{
+					ShouldIncludeInTotal = false;
+					break;
+				}
+			}
+			if (ShouldIncludeInTotal == false)
+			{
+				OutExcludedIndices->Add(CurrentIndex);
+				continue;
+			}
+			OutIncrementallyRefinedUniqueIndices->Add(CurrentIndex);
+
+			for (int32 ChildIndex : OwningTable->GetAsset(CurrentIndex)->GetDependencies())
+			{
+				// Don't revisit nodes we've already visited and don't re-add ThisIndex to avoid loops (and to avoid counting ourself)
+				if (!OutIncrementallyRefinedUniqueIndices->Contains(ChildIndex) && ThisIndex != ChildIndex)
+				{
+					IndicesToVisit.AddUnique(ChildIndex);
+				}
 			}
 		}
-	}
 
-	return OutIncrementallyRefinedUniqueIndices->Num() != PreviouslyVisitedIndices.Num();
+		return OutIncrementallyRefinedUniqueIndices->Num() != PreviouslyVisitedIndices.Num();
+	}
+	else
+	{
+		return false;
+	}
 }
 
 void FAssetTableRow::ComputeDependencySizes(FAssetTable* OwningTable, int32 ThisIndex, TSet<int32>* OutUniqueDependencies/* = nullptr*/, TSet<int32>* OutSharedDependencies/* = nullptr*/) const
diff --git a/Engine/Plugins/Editor/AssetManagerEditor/Source/AssetManagerEditor/Private/TreeView/AssetTable.h b/Engine/Plugins/Editor/AssetManagerEditor/Source/AssetManagerEditor/Private/TreeView/AssetTable.h
--- a/Engine/Plugins/Editor/AssetManagerEditor/Source/AssetManagerEditor/Private/TreeView/AssetTable.h
+++ b/Engine/Plugins/Editor/AssetManagerEditor/Source/AssetManagerEditor/Private/TreeView/AssetTable.h
@@ -16,6 +16,8 @@ namespace UE
 	}
 }
 
+class FAssetTable;
+
 // Simple string store.
 class FAssetTableStringStore
 {
