AssetManagerEditor:
  - Exposed selection to AssetTableTreeView. Added status bar info with number of assets selected. When a single asset is selected, the path/name is displayed too in the status bar.
  - Updated the tree view presets.
  - When an asset is selected in the asset picker list view, it will be selected also in the tree view.
  - Removed non-async code path from STableTreeView and fixed few typos.
  - Made few minor performance improvements to ComputeDependencySizes implementation.
  - Fixed aggregated values for "Asset Count", "Staged Compressed Size" and "Total Usages" colums when a "Dependencies" node is expanded, by setting the IsDynamic flag for these columns.

diff --git a/Engine/Plugins/Editor/AssetManagerEditor/Source/AssetManagerEditor/Private/Insights/Filter/ViewModels/FilterConfigurator.cpp b/Engine/Plugins/Editor/AssetManagerEditor/Source/AssetManagerEditor/Private/Insights/Filter/ViewModels/FilterConfigurator.cpp
--- a/Engine/Plugins/Editor/AssetManagerEditor/Source/AssetManagerEditor/Private/Insights/Filter/ViewModels/FilterConfigurator.cpp
+++ b/Engine/Plugins/Editor/AssetManagerEditor/Source/AssetManagerEditor/Private/Insights/Filter/ViewModels/FilterConfigurator.cpp
@@ -49,7 +49,7 @@ FFilterConfigurator& FFilterConfigurator::operator=(const FFilterConfigurator& O
 	ComputeUsedKeys();
 	RootNode->ProcessFilter();
 
-	OnChangesCommitedEvent.Broadcast();
+	OnChangesCommittedEvent.Broadcast();
 
 	return *this;
 }
diff --git a/Engine/Plugins/Editor/AssetManagerEditor/Source/AssetManagerEditor/Private/Insights/Filter/ViewModels/FilterConfigurator.h b/Engine/Plugins/Editor/AssetManagerEditor/Source/AssetManagerEditor/Private/Insights/Filter/ViewModels/FilterConfigurator.h
--- a/Engine/Plugins/Editor/AssetManagerEditor/Source/AssetManagerEditor/Private/Insights/Filter/ViewModels/FilterConfigurator.h
+++ b/Engine/Plugins/Editor/AssetManagerEditor/Source/AssetManagerEditor/Private/Insights/Filter/ViewModels/FilterConfigurator.h
@@ -48,15 +48,15 @@ private:
 	FOnDestroyedEvent OnDestroyedEvent;
 
 	////////////////////////////////////////////////////////////////////////////////////////////////////
-	// OnChangesCommitedEvent
+	// OnChangesCommittedEvent
 
 public:
 	/** The event to execute when the changes to the Filter Widget are saved by clicking on the OK Button. */
-	DECLARE_MULTICAST_DELEGATE(FOnChangesCommitedEvent);
-	FOnChangesCommitedEvent& GetOnChangesCommitedEvent() { return OnChangesCommitedEvent; }
+	DECLARE_MULTICAST_DELEGATE(FOnChangesCommittedEvent);
+	FOnChangesCommittedEvent& GetOnChangesCommittedEvent() { return OnChangesCommittedEvent; }
 
 private:
-	FOnChangesCommitedEvent OnChangesCommitedEvent;
+	FOnChangesCommittedEvent OnChangesCommittedEvent;
 
 	////////////////////////////////////////////////////////////////////////////////////////////////////
 
diff --git a/Engine/Plugins/Editor/AssetManagerEditor/Source/AssetManagerEditor/Private/Insights/Table/ViewModels/TableColumn.h b/Engine/Plugins/Editor/AssetManagerEditor/Source/AssetManagerEditor/Private/Insights/Table/ViewModels/TableColumn.h
--- a/Engine/Plugins/Editor/AssetManagerEditor/Source/AssetManagerEditor/Private/Insights/Table/ViewModels/TableColumn.h
+++ b/Engine/Plugins/Editor/AssetManagerEditor/Source/AssetManagerEditor/Private/Insights/Table/ViewModels/TableColumn.h
@@ -25,12 +25,12 @@ class ITableCellValueSorter;
 
 enum class ETableColumnFlags : uint32
 {
-	None = 0,
-
-	ShouldBeVisible = (1 << 0),
-	CanBeHidden     = (1 << 1),
-	CanBeFiltered   = (1 << 2),
-	IsHierarchy     = (1 << 3),
+	None            = 0,
+	ShouldBeVisible = (1 << 0), // Column should be initially visible.
+	CanBeHidden     = (1 << 1), // Column can be hidden.
+	CanBeFiltered   = (1 << 2), // Column can be used for filtering.
+	IsHierarchy     = (1 << 3), // Column is the hierarchy (name) column, in a tree view.
+	IsDynamic       = (1 << 4), // Column is dynamic. The column's value can not be cached.
 };
 ENUM_CLASS_FLAGS(ETableColumnFlags);
 
@@ -42,8 +42,6 @@ enum class ETableColumnAggregation : uint32
 	Sum,
 	Min,
 	Max,
-	//Average,
-	//Median,
 	SameValue
 };
 
@@ -66,7 +64,6 @@ public:
 		, TitleName()
 		, Description()
 		, bIsVisible(false)
-		, bIsDynamic(false)
 		, Flags(ETableColumnFlags::None)
 		, HorizontalAlignment(HAlign_Left)
 		, InitialWidth(60.0f)
@@ -117,12 +114,15 @@ public:
 	/** Whether this column can be hidden. */
 	bool CanBeHidden() const { return EnumHasAnyFlags(Flags, ETableColumnFlags::CanBeHidden); }
 
-	/** Whether this column can be used for filtering displayed results. */
+	/** Whether this column can be used for filtering. */
 	bool CanBeFiltered() const { return EnumHasAnyFlags(Flags, ETableColumnFlags::CanBeFiltered); }
 
-	/** Whether this column is the hierarcy (name) column, in a tree view. */
+	/** Whether this column is the hierarchy (name) column, in a tree view. */
 	bool IsHierarchy() const { return EnumHasAnyFlags(Flags, ETableColumnFlags::IsHierarchy); }
 
+	/** Whether this column is dynamic or not. If a column is dynamic, the column's value can not be cached. */
+	bool IsDynamic() const { return EnumHasAnyFlags(Flags, ETableColumnFlags::IsDynamic); }
+
 	void SetFlags(ETableColumnFlags InFlags) { Flags = InFlags; }
 
 	//////////////////////////////////////////////////
@@ -211,9 +211,6 @@ public:
 
 	//////////////////////////////////////////////////
 
-	void SetIsDynamic(bool InValue) { bIsDynamic = InValue; }
-	bool IsDynamic() const { return bIsDynamic; }
-
 private:
 	/** Id of the column. */
 	FName Id;
@@ -233,9 +230,6 @@ private:
 	/** Is this column visible? */
 	bool bIsVisible;
 
-	/** If it is not dynamic, the column's value can be cached. */
-	bool bIsDynamic;
-
 	/** Other on/off switches. */
 	ETableColumnFlags Flags;
 
@@ -256,7 +250,7 @@ private:
 	/** Custom getter for values identified by this column. */
 	TSharedRef<ITableCellValueGetter> ValueGetter;
 
-	/** Custom formater for values displayed by this column. */
+	/** Custom formatter for values displayed by this column. */
 	TSharedRef<ITableCellValueFormatter> ValueFormatter;
 
 	/** Custom sorter for values displayed by this column. */
@@ -265,7 +259,7 @@ private:
 	/** Initial sorting mode. */
 	EColumnSortMode::Type InitialSortMode;
 
-	/** Used to convert in a custom way from string to column data type in FilterConfigurator */
+	/** Used in FilterConfigurator to convert from string to column's data type. */
 	TSharedPtr<IFilterValueConverter> ValueConverter;
 
 	/* Parent table. Only one table instance can own this column. */
diff --git a/Engine/Plugins/Editor/AssetManagerEditor/Source/AssetManagerEditor/Private/Insights/Table/Widgets/STableTreeView.cpp b/Engine/Plugins/Editor/AssetManagerEditor/Source/AssetManagerEditor/Private/Insights/Table/Widgets/STableTreeView.cpp
--- a/Engine/Plugins/Editor/AssetManagerEditor/Source/AssetManagerEditor/Private/Insights/Table/Widgets/STableTreeView.cpp
+++ b/Engine/Plugins/Editor/AssetManagerEditor/Source/AssetManagerEditor/Private/Insights/Table/Widgets/STableTreeView.cpp
@@ -108,13 +108,13 @@ public:
 			FInputChord(EModifierKey::Control, EKeys::E));
 
 		UI_COMMAND(Command_ExportEntireTreeToFile,
-			"Export Entire Tree to File...",
+			"Export Entire Tree (+ Leaf Nodes) to File...",
 			"Exports the entire tree/table content to a file. It exports also the collapsed tree nodes, including the leaf nodes. Filtered out nodes are not exported.",
 			EUserInterfaceActionType::Button,
 			FInputChord(EModifierKey::Shift, EKeys::E));
 
-		UI_COMMAND(Command_ExportEntireTreeToFileNoLeafs,
-			"Export Entire Tree (No Leafs) to File...",
+		UI_COMMAND(Command_ExportEntireTreeToFileExceptLeaves,
+			"Export Entire Tree (- Leaf Nodes) to File...",
 			"Exports the entire tree/table content to a file, but not the leaf nodes. It exports the collapsed tree nodes. Filtered out nodes are not exported.",
 			EUserInterfaceActionType::Button,
 			FInputChord(EModifierKey::Alt, EKeys::E));
@@ -129,7 +129,7 @@ public:
 	TSharedPtr<FUICommandInfo> Command_CollapseSubtree;
 	TSharedPtr<FUICommandInfo> Command_ExportToFile;
 	TSharedPtr<FUICommandInfo> Command_ExportEntireTreeToFile;
-	TSharedPtr<FUICommandInfo> Command_ExportEntireTreeToFileNoLeafs;
+	TSharedPtr<FUICommandInfo> Command_ExportEntireTreeToFileExceptLeaves;
 };
 
 ////////////////////////////////////////////////////////////////////////////////////////////////////
@@ -141,30 +141,9 @@ const FName STableTreeView::RootNodeName(TEXT("Root"));
 ////////////////////////////////////////////////////////////////////////////////////////////////////
 
 STableTreeView::STableTreeView()
-	: Table()
-	, TreeView(nullptr)
-	, TreeViewHeaderRow(nullptr)
-	, ExternalScrollbar(nullptr)
-	, HoveredColumnId()
-	, HoveredNodePtr(nullptr)
-	, HighlightedNodeName()
-	, Root(MakeShared<FTableTreeNode>(RootNodeName, Table))
-	, TableTreeNodes()
-	, FilteredGroupNodes()
-	, ExpandedNodes()
-	, bExpansionSaved(false)
-	, SearchBox(nullptr)
-	, TextFilter(nullptr)
-	, Filters(nullptr)
-	, AvailableGroupings()
-	, CurrentGroupings()
-	, GroupingBreadcrumbTrail(nullptr)
-	, AvailableSorters()
-	, CurrentSorter(nullptr)
+	: Root(MakeShared<FTableTreeNode>(RootNodeName, Table))
 	, ColumnBeingSorted(GetDefaultColumnBeingSorted())
 	, ColumnSortMode(GetDefaultColumnSortMode())
-	, StatsStartTime(0.0)
-	, StatsEndTime(0.0)
 {
 }
 
@@ -172,10 +151,7 @@ STableTreeView::STableTreeView()
 
 STableTreeView::~STableTreeView()
 {
-	if (bRunInAsyncMode)
-	{
-		checkf(bIsCloseScheduled, TEXT("TableTreeView running in async mode was closed but OnClose() was not called. This can lead to a crash. Call OnClose() from the owner tab/window."))
-	}
+	checkf(bIsCloseScheduled, TEXT("TableTreeView running in async mode was closed but OnClose() was not called. This can lead to a crash. Call OnClose() from the owner tab/window."))
 
 	if (CurrentAsyncOpFilterConfigurator)
 	{
@@ -225,7 +201,7 @@ void STableTreeView::InitCommandList()
 		FExecuteAction::CreateSP(this, &STableTreeView::ContextMenu_ExportToFile_Execute, true, true),
 		FCanExecuteAction::CreateSP(this, &STableTreeView::ContextMenu_ExportToFile_CanExecute));
 
-	CommandList->MapAction(FTableTreeViewCommands::Get().Command_ExportEntireTreeToFileNoLeafs,
+	CommandList->MapAction(FTableTreeViewCommands::Get().Command_ExportEntireTreeToFileExceptLeaves,
 		FExecuteAction::CreateSP(this, &STableTreeView::ContextMenu_ExportToFile_Execute, true, false),
 		FCanExecuteAction::CreateSP(this, &STableTreeView::ContextMenu_ExportToFile_CanExecute));
 }
@@ -355,15 +331,8 @@ void STableTreeView::ConstructWidget(TSharedPtr<FTable> InTablePtr)
 	// Create the search filters: text based, type based etc.
 	TextFilter = MakeShared<FTableTreeNodeTextFilter>(FTableTreeNodeTextFilter::FItemToStringArray::CreateStatic(&STableTreeView::HandleItemToStringArray));
 	Filters = MakeShared<FTableTreeNodeFilterCollection>();
-	if (bRunInAsyncMode)
-	{
-		CurrentAsyncOpTextFilter = MakeShared<FTableTreeNodeTextFilter>(FTableTreeNodeTextFilter::FItemToStringArray::CreateStatic(&STableTreeView::HandleItemToStringArray));
-		Filters->Add(CurrentAsyncOpTextFilter);
-	}
-	else
-	{
-		Filters->Add(TextFilter);
-	}
+	CurrentAsyncOpTextFilter = MakeShared<FTableTreeNodeTextFilter>(FTableTreeNodeTextFilter::FItemToStringArray::CreateStatic(&STableTreeView::HandleItemToStringArray));
+	Filters->Add(CurrentAsyncOpTextFilter);
 
 	InitializeAndShowHeaderColumns();
 
@@ -524,8 +493,8 @@ END_SLATE_FUNCTION_BUILD_OPTIMIZATION
 
 TSharedPtr<SWidget> STableTreeView::TreeView_GetMenuContent()
 {
-	const TArray<FTableTreeNodePtr> SelectedNodes = TreeView->GetSelectedItems();
-	const int32 NumSelectedNodes = SelectedNodes.Num();
+	TArray<FTableTreeNodePtr> SelectedNodes;
+	const int32 NumSelectedNodes = TreeView->GetSelectedItems(SelectedNodes);
 	FTableTreeNodePtr SelectedNode = NumSelectedNodes ? SelectedNodes[0] : nullptr;
 
 	FText SelectionStr;
@@ -778,7 +747,7 @@ void STableTreeView::TreeView_BuildExportMenu(FMenuBuilder& MenuBuilder)
 
 		MenuBuilder.AddMenuEntry
 		(
-			FTableTreeViewCommands::Get().Command_ExportEntireTreeToFileNoLeafs,
+			FTableTreeViewCommands::Get().Command_ExportEntireTreeToFileExceptLeaves,
 			NAME_None,
 			TAttribute<FText>(),
 			TAttribute<FText>(),
@@ -962,7 +931,7 @@ TSharedRef<SWidget> STableTreeView::TreeViewHeaderRow_GenerateColumnMenu(const F
 
 void STableTreeView::Tick(const FGeometry& AllottedGeometry, const double InCurrentTime, const float InDeltaTime)
 {
-	if (bRunInAsyncMode && bIsUpdateRunning && !bIsCloseScheduled)
+	if (bIsUpdateRunning && !bIsCloseScheduled)
 	{
 		if (DispatchEvent.IsValid() && !DispatchEvent->IsComplete())
 		{
@@ -984,32 +953,29 @@ void STableTreeView::Tick(const FGeometry& AllottedGeometry, const double InCurr
 			}
 		}
 	}
+
+	if (LazyInvalidateWidgetReason != EInvalidateWidgetReason::None)
+	{
+		TreeView->Invalidate(LazyInvalidateWidgetReason);
+		LazyInvalidateWidgetReason = EInvalidateWidgetReason::None;
+	}
 }
 
 ////////////////////////////////////////////////////////////////////////////////////////////////////
 
 void STableTreeView::UpdateTree()
 {
-	if (bRunInAsyncMode)
+	if (!bIsUpdateRunning)
 	{
-		if (!bIsUpdateRunning)
-		{
-			OnPreAsyncUpdate();
+		OnPreAsyncUpdate();
 
-			FGraphEventRef CompletedEvent = StartCreateGroupsTask();
-			CompletedEvent = StartSortTreeNodesTask(CompletedEvent);
-			InProgressAsyncOperationEvent = StartApplyFiltersTask(CompletedEvent);
-		}
-		else
-		{
-			CancelCurrentAsyncOp();
-		}
+		FGraphEventRef CompletedEvent = StartCreateGroupsTask();
+		CompletedEvent = StartSortTreeNodesTask(CompletedEvent);
+		InProgressAsyncOperationEvent = StartApplyFiltersTask(CompletedEvent);
 	}
 	else
 	{
-		ApplyGrouping();
-		ApplySorting();
-		ApplyFiltering();
+		CancelCurrentAsyncOp();
 	}
 }
 
@@ -1019,22 +985,15 @@ void STableTreeView::UpdateTree()
 
 void STableTreeView::OnFilteringChanged()
 {
-	if (bRunInAsyncMode)
+	if (!bIsUpdateRunning)
 	{
-		if (!bIsUpdateRunning)
-		{
-			OnPreAsyncUpdate();
+		OnPreAsyncUpdate();
 
-			InProgressAsyncOperationEvent = StartApplyFiltersTask();
-		}
-		else
-		{
-			CancelCurrentAsyncOp();
-		}
+		InProgressAsyncOperationEvent = StartApplyFiltersTask();
 	}
 	else
 	{
-		ApplyFiltering();
+		CancelCurrentAsyncOp();
 	}
 }
 
@@ -1059,7 +1018,7 @@ void STableTreeView::ApplyFiltering(FTableTreeNodePtr SubTreeRootPtr)
 	}
 	else
 	{
-		const bool bFilterIsEmpty = bRunInAsyncMode ? CurrentAsyncOpTextFilter->GetRawFilterText().IsEmpty() : TextFilter->GetRawFilterText().IsEmpty();
+		const bool bFilterIsEmpty = CurrentAsyncOpTextFilter->GetRawFilterText().IsEmpty();
 		ApplyHierarchicalFilterForNode(SubTreeRootPtr, bFilterIsEmpty);
 	}
 
@@ -1094,44 +1053,6 @@ void STableTreeView::ApplyFiltering(FTableTreeNodePtr SubTreeRootPtr)
 	}
 
 	UpdateAggregatedValuesRec(*Root);
-
-	// Cannot call TreeView functions from other threads than MainThread and SlateThread.
-	if (!bRunInAsyncMode)
-	{
-		// Only expand nodes if we have a text filter.
-		const bool bNonEmptyTextFilter = !TextFilter->GetRawFilterText().IsEmpty();
-		if (bNonEmptyTextFilter)
-		{
-			if (!bExpansionSaved)
-			{
-				ExpandedNodes.Empty();
-				TreeView->GetExpandedItems(ExpandedNodes);
-				bExpansionSaved = true;
-			}
-
-			for (int32 Fx = 0; Fx < FilteredGroupNodes.Num(); Fx++)
-			{
-				const FTableTreeNodePtr& GroupPtr = FilteredGroupNodes[Fx];
-				TreeView->SetItemExpansion(GroupPtr, GroupPtr->IsExpanded());
-			}
-		}
-		else
-		{
-			if (bExpansionSaved)
-			{
-				// Restore previously expanded nodes when the text filter is disabled.
-				TreeView->ClearExpandedItems();
-				for (auto It = ExpandedNodes.CreateConstIterator(); It; ++It)
-				{
-					TreeView->SetItemExpansion(*It, true);
-				}
-				bExpansionSaved = false;
-			}
-		}
-
-		// Request tree refresh
-		TreeView->RequestTreeRefresh();
-	}
 }
 
 ////////////////////////////////////////////////////////////////////////////////////////////////////
@@ -1174,14 +1095,7 @@ bool STableTreeView::ApplyHierarchicalFilterForNode(FTableTreeNodePtr NodePtr, b
 		const bool bIsGroupNodeVisible = bIsNodeVisible || NumVisibleChildren > 0;
 		if (!bFilterIsEmpty && bIsGroupNodeVisible)
 		{
-			if (bRunInAsyncMode)
-			{
-				NodesToExpand.Add(NodePtr);
-			}
-			else
-			{
-				TreeView->SetItemExpansion(NodePtr, true);
-			}
+			NodesToExpand.Add(NodePtr);
 		}
 		return bIsGroupNodeVisible;
 	}
@@ -1270,14 +1184,7 @@ bool STableTreeView::MakeSubtreeVisible(FTableTreeNodePtr NodePtr, bool bFilterI
 
 			if (bShouldExpand)
 			{
-				if (bRunInAsyncMode)
-				{
-					NodesToExpand.Add(NodePtr);
-				}
-				else
-				{
-					TreeView->SetItemExpansion(NodePtr, true);
-				}
+				NodesToExpand.Add(NodePtr);
 			}
 		}
 
@@ -1467,26 +1374,17 @@ FText STableTreeView::SearchBox_GetTooltipText() const
 
 void STableTreeView::OnGroupingChanged()
 {
-	if (bRunInAsyncMode)
+	if (!bIsUpdateRunning)
 	{
-		if (!bIsUpdateRunning)
-		{
-			OnPreAsyncUpdate();
+		OnPreAsyncUpdate();
 
-			FGraphEventRef CompletedEvent = StartCreateGroupsTask();
-			CompletedEvent = StartSortTreeNodesTask(CompletedEvent);
-			InProgressAsyncOperationEvent = StartApplyFiltersTask(CompletedEvent);
-		}
-		else
-		{
-			CancelCurrentAsyncOp();
-		}
+		FGraphEventRef CompletedEvent = StartCreateGroupsTask();
+		CompletedEvent = StartSortTreeNodesTask(CompletedEvent);
+		InProgressAsyncOperationEvent = StartApplyFiltersTask(CompletedEvent);
 	}
 	else
 	{
-		ApplyGrouping();
-		ApplySorting();
-		ApplyFiltering();
+		CancelCurrentAsyncOp();
 	}
 }
 
@@ -1588,17 +1486,6 @@ void STableTreeView::GroupNodesRec(const TArray<FTableTreeNodePtr>& Nodes, FTabl
 			}
 		}
 	}
-
-	if (!bRunInAsyncMode)
-	{
-		// Expand group nodes.
-		for (FBaseTreeNodePtr GroupPtr : ParentGroup.GetChildren())
-		{
-			ensure(GroupPtr->IsGroup());
-			FTableTreeNodePtr TableTreeGroupPtr = StaticCastSharedPtr<FTableTreeNode>(GroupPtr);
-			TreeView->SetItemExpansion(TableTreeGroupPtr, TableTreeGroupPtr->IsExpanded());
-		}
-	}
 }
 
 ////////////////////////////////////////////////////////////////////////////////////////////////////
@@ -2473,24 +2360,16 @@ void STableTreeView::UpdateCurrentSortingByColumn()
 
 void STableTreeView::OnSortingChanged()
 {
-	if (bRunInAsyncMode)
+	if (!bIsUpdateRunning)
 	{
-		if (!bIsUpdateRunning)
-		{
-			OnPreAsyncUpdate();
+		OnPreAsyncUpdate();
 
-			FGraphEventRef CompletedEvent = StartSortTreeNodesTask();
-			InProgressAsyncOperationEvent = StartApplyFiltersTask(CompletedEvent);
-		}
-		else
-		{
-			CancelCurrentAsyncOp();
-		}
+		FGraphEventRef CompletedEvent = StartSortTreeNodesTask();
+		InProgressAsyncOperationEvent = StartApplyFiltersTask(CompletedEvent);
 	}
 	else
 	{
-		ApplySorting();
-		ApplyFiltering();
+		CancelCurrentAsyncOp();
 	}
 }
 
@@ -2911,9 +2790,6 @@ void STableTreeView::ContextMenu_HideAllColumns_Execute()
 
 void STableTreeView::Reset()
 {
-	StatsStartTime = 0.0;
-	StatsEndTime = 0.0;
-
 	RebuildTree(true);
 }
 
@@ -3011,7 +2887,7 @@ void STableTreeView::OnPostAsyncUpdate()
 		ExpandedNodes.Empty();
 
 		// Expand each group node on the first few depths (if it doesn't have too many children).
-		SetExpandValueForChildGroups(Root.Get(), 1000, 4, true);
+		SetExpandValueForChildGroups(Root.Get(), MaxNodesToAutoExpand, MaxDepthToAutoExpand, true);
 
 		ClearInProgressAsyncOperations();
 		TreeView_Refresh();
@@ -3023,11 +2899,11 @@ void STableTreeView::OnPostAsyncUpdate()
 
 ////////////////////////////////////////////////////////////////////////////////////////////////////
 
-void STableTreeView::SetExpandValueForChildGroups(FBaseTreeNode* InRoot, int32 InMaxExpandedNodes, int32 InMaxDepthToExpand, bool InValue)
+void STableTreeView::SetExpandValueForChildGroups(FBaseTreeNode* InRoot, int32 InMaxNodesToExpand, int32 InMaxDepthToExpand, bool InValue)
 {
 	TArray<int32> NumNodesPerDepth;
 	NumNodesPerDepth.AddDefaulted(InMaxDepthToExpand + 1);
-	CountNumNodesPerDepthRec(InRoot, NumNodesPerDepth, 0, InMaxDepthToExpand, InMaxExpandedNodes);
+	CountNumNodesPerDepthRec(InRoot, NumNodesPerDepth, 0, InMaxDepthToExpand, InMaxNodesToExpand);
 
 	int32 MaxDepth = 0;
 	for (int32 Depth = 0; Depth <= InMaxDepthToExpand; ++Depth)
@@ -3036,7 +2912,7 @@ void STableTreeView::SetExpandValueForChildGroups(FBaseTreeNode* InRoot, int32 I
 		{
 			NumNodesPerDepth[Depth] += NumNodesPerDepth[Depth - 1];
 		}
-		if (NumNodesPerDepth[Depth] > InMaxExpandedNodes)
+		if (NumNodesPerDepth[Depth] > InMaxNodesToExpand)
 		{
 			break;
 		}
@@ -3324,7 +3200,7 @@ FReply STableTreeView::OnAdvancedFiltersClicked()
 		AddCustomAdvancedFilters();
 
 		CurrentAsyncOpFilterConfigurator = new FFilterConfigurator(*FilterConfigurator);
-		OnFilterChangesCommitedHandle = FilterConfigurator->GetOnChangesCommitedEvent().AddSP(this, &STableTreeView::OnAdvancedFiltersChangesCommited);
+		OnFilterChangesCommittedHandle = FilterConfigurator->GetOnChangesCommittedEvent().AddSP(this, &STableTreeView::OnAdvancedFiltersChangesCommitted);
 	}
 
 	FFilterService::Get()->CreateFilterConfiguratorWidget(FilterConfigurator);
@@ -3336,7 +3212,7 @@ FReply STableTreeView::OnAdvancedFiltersClicked()
 
 bool STableTreeView::ApplyAdvancedFilters(const FTableTreeNodePtr& NodePtr)
 {
-	FFilterConfigurator* FilterConfiguratorToUse = bRunInAsyncMode ? CurrentAsyncOpFilterConfigurator : FilterConfigurator.Get();
+	FFilterConfigurator* FilterConfiguratorToUse = CurrentAsyncOpFilterConfigurator;
 
 	if (FilterConfiguratorToUse == nullptr)
 	{
@@ -3389,7 +3265,7 @@ bool STableTreeView::ApplyAdvancedFilters(const FTableTreeNodePtr& NodePtr)
 
 ////////////////////////////////////////////////////////////////////////////////////////////////////
 
-void STableTreeView::OnAdvancedFiltersChangesCommited()
+void STableTreeView::OnAdvancedFiltersChangesCommitted()
 {
 	OnFilteringChanged();
 }
@@ -3540,14 +3416,19 @@ bool STableTreeView::ContextMenu_ExpandSubtree_CanExecute() const
 
 void STableTreeView::ContextMenu_ExpandSubtree_Execute()
 {
-	const TArray<FTableTreeNodePtr> SelectedNodes = TreeView->GetSelectedItems();
+	TArray<FTableTreeNodePtr> SelectedNodes;
+	if (!TreeView->GetSelectedItems(SelectedNodes))
+	{
+		return;
+	}
+
 	for (const FTableTreeNodePtr& Node : SelectedNodes)
 	{
 		if (Node->IsGroup())
 		{
 			Node->SetExpansion(true);
 			TreeView->SetItemExpansion(Node, true);
-			SetExpandValueForChildGroups((FBaseTreeNode*)Node.Get(), MAX_NUMBER_OF_NODES_TO_EXPAND, MAX_DEPTH_TO_EXPAND, true);
+			SetExpandValueForChildGroups((FBaseTreeNode*)Node.Get(), MaxNodesToExpand, MaxDepthToExpand, true);
 		}
 	}
 
@@ -3558,7 +3439,12 @@ void STableTreeView::ContextMenu_ExpandSubtree_Execute()
 
 bool STableTreeView::ContextMenu_CollapseSubtree_CanExecute() const
 {
-	const TArray<FTableTreeNodePtr> SelectedNodes = TreeView->GetSelectedItems();
+	TArray<FTableTreeNodePtr> SelectedNodes;
+	if (!TreeView->GetSelectedItems(SelectedNodes))
+	{
+		return false;
+	}
+
 	for (const FTableTreeNodePtr& Node : SelectedNodes)
 	{
 		if (Node->IsGroup() && Node->GetFilteredChildren().Num() > 0 && TreeView->IsItemExpanded(Node))
@@ -3574,14 +3460,19 @@ bool STableTreeView::ContextMenu_CollapseSubtree_CanExecute() const
 
 void STableTreeView::ContextMenu_CollapseSubtree_Execute()
 {
-	const TArray<FTableTreeNodePtr> SelectedNodes = TreeView->GetSelectedItems();
+	TArray<FTableTreeNodePtr> SelectedNodes;
+	if (!TreeView->GetSelectedItems(SelectedNodes))
+	{
+		return;
+	}
+
 	for (const FTableTreeNodePtr& Node : SelectedNodes)
 	{
 		if (Node->IsGroup() && TreeView->IsItemExpanded(Node))
 		{
 			Node->SetExpansion(false);
 			TreeView->SetItemExpansion(Node, false);
-			SetExpandValueForChildGroups((FBaseTreeNode*)Node.Get(), MAX_NUMBER_OF_NODES_TO_EXPAND, MAX_DEPTH_TO_EXPAND, false);
+			SetExpandValueForChildGroups((FBaseTreeNode*)Node.Get(), MaxNodesToExpand, MaxDepthToExpand, false);
 		}
 	}
 
@@ -3599,7 +3490,12 @@ bool STableTreeView::ContextMenu_ExpandCriticalPath_CanExecute() const
 
 void STableTreeView::ContextMenu_ExpandCriticalPath_Execute()
 {
-	const TArray<FTableTreeNodePtr> SelectedNodes = TreeView->GetSelectedItems();
+	TArray<FTableTreeNodePtr> SelectedNodes;
+	if (!TreeView->GetSelectedItems(SelectedNodes))
+	{
+		return;
+	}
+
 	for (const FTableTreeNodePtr& Node : SelectedNodes)
 	{
 		FTableTreeNodePtr CurrentNode = Node;
@@ -3776,7 +3672,7 @@ void STableTreeView::ApplyViewPreset(const ITableTreeViewPreset& InPreset)
 void STableTreeView::ApplyColumnConfig(const TArrayView<FTableColumnConfig >& InColumnConfigSet)
 {
 	// TODO: Reorder columns as in the config set.
-	// Currenly we only apply visibility and column width.
+	// Currently we only apply visibility and column width.
 	for (const TSharedRef<FTableColumn>& ColumnRef : Table->GetColumns())
 	{
 		FTableColumn& Column = ColumnRef.Get();
diff --git a/Engine/Plugins/Editor/AssetManagerEditor/Source/AssetManagerEditor/Private/Insights/Table/Widgets/STableTreeView.h b/Engine/Plugins/Editor/AssetManagerEditor/Source/AssetManagerEditor/Private/Insights/Table/Widgets/STableTreeView.h
--- a/Engine/Plugins/Editor/AssetManagerEditor/Source/AssetManagerEditor/Private/Insights/Table/Widgets/STableTreeView.h
+++ b/Engine/Plugins/Editor/AssetManagerEditor/Source/AssetManagerEditor/Private/Insights/Table/Widgets/STableTreeView.h
@@ -243,10 +243,10 @@ protected:
 	void TreeView_OnGetChildren(FTableTreeNodePtr InParent, TArray<FTableTreeNodePtr>& OutChildren);
 
 	/** Called by STreeView when selection has changed. */
-	void TreeView_OnSelectionChanged(FTableTreeNodePtr SelectedItem, ESelectInfo::Type SelectInfo);
+	virtual void TreeView_OnSelectionChanged(FTableTreeNodePtr SelectedItem, ESelectInfo::Type SelectInfo);
 
 	/** Called by STreeView when a tree node is expanded or collapsed. */
-	void TreeView_OnExpansionChanged(FTableTreeNodePtr TreeNode, bool bShouldBeExpanded);
+	virtual void TreeView_OnExpansionChanged(FTableTreeNodePtr TreeNode, bool bShouldBeExpanded);
 
 	/** Called by STreeView when a tree item is double clicked. */
 	virtual void TreeView_OnMouseButtonDoubleClick(FTableTreeNodePtr TreeNode);
@@ -417,7 +417,7 @@ protected:
 	bool AdvancedFilters_ShouldBeEnabled() const;
 	FText AdvancedFilters_GetTooltipText() const;
 	bool FilterConfigurator_HasFilters() const;
-	void OnAdvancedFiltersChangesCommited();
+	void OnAdvancedFiltersChangesCommitted();
 	bool ApplyAdvancedFilters(const FTableTreeNodePtr& NodePtr);
 	bool virtual ApplyCustomAdvancedFilters(const FTableTreeNodePtr& NodePtr) { return true; };
 	virtual void AddCustomAdvancedFilters() {}
@@ -465,8 +465,11 @@ protected:
 protected:
 	/** Table view model. */
 	TSharedPtr<FTable> Table;
+
+	TSharedPtr<FUICommandList> CommandList;
+
 	//////////////////////////////////////////////////
-	// Tree View, Columns
+	// Widget
 
 	/** The child STreeView widget. */
 	TSharedPtr<STreeView<FTableTreeNodePtr>> TreeView;
@@ -489,8 +492,6 @@ protected:
 	/** Name of the tree node that should be drawn as highlighted. */
 	FName HighlightedNodeName;
 
-	TSharedPtr<FUICommandList> CommandList;
-
 	//////////////////////////////////////////////////
 	// Tree Nodes
 
@@ -508,8 +509,12 @@ protected:
 	/** Currently expanded group nodes. */
 	TSet<FTableTreeNodePtr> ExpandedNodes;
 
-	/** If true, the expanded nodes have been saved before applying a text filter. */
-	bool bExpansionSaved;
+	static constexpr int32 MaxNodesToAutoExpand = 1000;
+	static constexpr int32 MaxDepthToAutoExpand = 3;
+	static constexpr int32 MaxNodesToExpand = 1000000;
+	static constexpr int32 MaxDepthToExpand = 100;
+
+	EInvalidateWidgetReason LazyInvalidateWidgetReason = EInvalidateWidgetReason::None;
 
 	//////////////////////////////////////////////////
 	// Search box and filters
@@ -523,6 +528,10 @@ protected:
 	/** The filter collection. */
 	TSharedPtr<FTableTreeNodeFilterCollection> Filters;
 
+	TSharedPtr<FFilterConfigurator> FilterConfigurator;
+	FDelegateHandle OnFilterChangesCommittedHandle;
+	FFilterContext Context;
+
 	//////////////////////////////////////////////////
 	// Grouping
 
@@ -546,12 +555,11 @@ protected:
 	FName ColumnBeingSorted;
 
 	/** How we sort the nodes? Ascending or Descending. */
-	EColumnSortMode::Type ColumnSortMode;
+	EColumnSortMode::Type ColumnSortMode = EColumnSortMode::None;
 
 	//////////////////////////////////////////////////
-	// Async
+	// Async Operations
 
-	bool bRunInAsyncMode = false;
 	bool bIsUpdateRunning = false;
 	bool bIsCloseScheduled = false;
 
@@ -574,22 +582,12 @@ protected:
 
 	//////////////////////////////////////////////////
 
-	TSharedPtr<FFilterConfigurator> FilterConfigurator;
-	FDelegateHandle OnFilterChangesCommitedHandle;
-	FFilterContext Context;
-
-	double StatsStartTime;
-	double StatsEndTime;
-
 	FText TreeViewBannerText;
 
 	TArray<TSharedRef<ITableTreeViewPreset>> AvailableViewPresets;
 	TSharedPtr<ITableTreeViewPreset> SelectedViewPreset;
 	TSharedPtr<SComboBox<TSharedRef<ITableTreeViewPreset>>> PresetComboBox;
 
-	static constexpr int32 MAX_NUMBER_OF_NODES_TO_EXPAND = 1000 * 1000;
-	static constexpr int32 MAX_DEPTH_TO_EXPAND = 100;
-
 	//////////////////////////////////////////////////
 	// Logging
 
@@ -697,7 +695,7 @@ public:
 
 	void DoTask(ENamedThreads::Type CurrentThread, const FGraphEventRef& MyCompletionGraphEvent)
 	{
-		// The role of this task is to keep the STableTreeView object alive until the task and it's prerequisits are completed and to destroy it on the game thread.
+		// The role of this task is to keep the STableTreeView object alive until the task and it's prerequisites are completed and to destroy it on the game thread.
 		FGraphEventRef Event = TableTreeViewPtr->AsyncCompleteTaskEvent;
 		if (TableTreeViewPtr.IsValid())
 		{
diff --git a/Engine/Plugins/Editor/AssetManagerEditor/Source/AssetManagerEditor/Private/SAssetAuditBrowser.cpp b/Engine/Plugins/Editor/AssetManagerEditor/Source/AssetManagerEditor/Private/SAssetAuditBrowser.cpp
--- a/Engine/Plugins/Editor/AssetManagerEditor/Source/AssetManagerEditor/Private/SAssetAuditBrowser.cpp
+++ b/Engine/Plugins/Editor/AssetManagerEditor/Source/AssetManagerEditor/Private/SAssetAuditBrowser.cpp
@@ -26,6 +26,7 @@
 #include "AssetManagerEditorCommands.h"
 #include "Insights/Common/Log.h"
 #include "TreeView/AssetTable.h"
+#include "TreeView/AssetTreeNode.h"
 #include "TreeView/SAssetTableTreeView.h"
 
 #define LOCTEXT_NAMESPACE "AssetManagementBrowser"
@@ -156,13 +157,13 @@ void SAssetAuditBrowser::FindReferencesForSelectedAssets() const
 	IAssetManagerEditorModule::ExtractAssetIdentifiersFromAssetDataList(Assets, AssetIdentifiers);
 
 	if (AssetIdentifiers.Num() > 0)
-{
+	{
 		IAssetManagerEditorModule::Get().OpenReferenceViewerUI(AssetIdentifiers);
 	}
 }
 
 void SAssetAuditBrowser::ShowSizeMapForSelectedAssets() const
-	{
+{
 	TArray<FAssetData> Assets = GetCurrentSelectionDelegate.Execute();
 	TArray<FAssetIdentifier> AssetIdentifiers;
 	IAssetManagerEditorModule::ExtractAssetIdentifiersFromAssetDataList(Assets, AssetIdentifiers);
@@ -244,7 +245,7 @@ void SAssetAuditBrowser::Construct(const FArguments& InArgs)
 		FExecuteAction::CreateSP(this, &SAssetAuditBrowser::EditSelectedAssets),
 		FCanExecuteAction::CreateSP(this, &SAssetAuditBrowser::IsAnythingSelected)
 	));
-	
+
 	Commands->MapAction(FAssetManagerEditorCommands::Get().ViewReferences, FUIAction(
 		FExecuteAction::CreateSP(this, &SAssetAuditBrowser::FindReferencesForSelectedAssets),
 		FCanExecuteAction::CreateSP(this, &SAssetAuditBrowser::IsAnythingSelected)
@@ -275,12 +276,13 @@ void SAssetAuditBrowser::Construct(const FArguments& InArgs)
 	Config.OnGetCustomSourceAssets = FOnGetCustomSourceAssets::CreateSP(this, &SAssetAuditBrowser::OnGetCustomSourceAssets);
 	Config.SyncToAssetsDelegates.Add(&SyncToAssetsDelegate);
 	Config.OnShouldFilterAsset = FOnShouldFilterAsset::CreateSP(this, &SAssetAuditBrowser::HandleFilterAsset);
+	Config.OnAssetSelected = FOnAssetSelected::CreateSP(this, &SAssetAuditBrowser::OnAssetSelected);
 	Config.GetCurrentSelectionDelegates.Add(&GetCurrentSelectionDelegate);
 	Config.SetFilterDelegates.Add(&SetFilterDelegate);
 	Config.bFocusSearchBoxWhenOpened = false;
 
 	Config.SaveSettingsName = SettingsIniSection;
-	
+
 	// Hide path and type by default
 	Config.HiddenColumnNames.Add(TEXT("Class"));
 	Config.HiddenColumnNames.Add(TEXT("Path"));
@@ -526,6 +528,7 @@ void SAssetAuditBrowser::Construct(const FArguments& InArgs)
 			+ SSplitter::Slot()
 			[
 				SAssignNew(AssetTableTreeView, SAssetTableTreeView, AssetTable)
+				.OnSelectionChanged(this, &SAssetAuditBrowser::TreeView_OnSelectionChanged)
 			]
 		];
 
@@ -539,7 +542,7 @@ void SAssetAuditBrowser::Construct(const FArguments& InArgs)
 		];
 	}
 
-	// Refresh our data sources to make sure we are displaying up to date info otherwise the user will need to 
+	// Refresh our data sources to make sure we are displaying up to date info otherwise the user will need to
 	// hit "refresh" every time they open the dialog to be able to trust it
 	ManagerEditorModule.RefreshRegistryData();
 }
@@ -730,15 +733,14 @@ void SAssetAuditBrowser::GoToHistoryIndex(int32 InHistoryIdx)
 	}
 }
 
-void SAssetAuditBrowser::PopulateAssetTableRow(FAssetTableRow& OutRow, const FAssetData& AssetData, FAssetTable& AssetTable) const
+void SAssetAuditBrowser::QuickPopulateAssetTableRow(FAssetTableRow& OutRow, const FAssetData& AssetData, FAssetTable& AssetTable) const
 {
 	// Asset Type
 	FString AssetType;
-	if (!EditorModule->GetStringValueForCustomColumn(AssetData, FName(TEXT("Type")), AssetType)) 
+	if (!EditorModule->GetStringValueForCustomColumn(AssetData, FName(TEXT("Type")), AssetType))
 	{
-		AssetType = AssetData.AssetClassPath.ToString(); 
+		AssetType = AssetData.AssetClassPath.ToString();
 	}
-	const uint32 AssetTypeHash = GetTypeHash(AssetType);
 	OutRow.Type = AssetTable.StoreStr(AssetType);
 
 	// Asset Name
@@ -753,6 +755,11 @@ void SAssetAuditBrowser::PopulateAssetTableRow(FAssetTableRow& OutRow, const FAs
 		ObjectPathString = ObjectPathString.Left(Index + 1);
 	}
 	OutRow.Path = AssetTable.StoreStr(ObjectPathString);
+}
+
+void SAssetAuditBrowser::PopulateAssetTableRow(FAssetTableRow& OutRow, const FAssetData& AssetData, FAssetTable& AssetTable) const
+{
+	QuickPopulateAssetTableRow(OutRow, AssetData, AssetTable);
 
 	FString Str;
 
@@ -781,9 +788,28 @@ void SAssetAuditBrowser::PopulateAssetTableRow(FAssetTableRow& OutRow, const FAs
 	OutRow.NativeClass = AssetTable.StoreStr(AssetData.AssetClassPath.GetAssetName().ToString());
 
 	// Sets the color based on asset type.
+	const uint32 AssetTypeHash = GetTypeHash(FStringView(OutRow.Type));
 	OutRow.Color = USlateThemeManager::Get().GetColor((EStyleColor)((uint32)EStyleColor::AccentBlue + AssetTypeHash % 8));
 }
 
+void SAssetAuditBrowser::TreeView_OnSelectionChanged(const TArray<TSharedPtr<UE::Insights::FTableTreeNode>> InSelectedNodes)
+{
+	if (!AssetTableTreeView.IsValid())
+	{
+		return;
+	}
+
+	//TODO: Find and select the corresponding asset in the Asset Picker list view.
+	TSharedPtr<FAssetTreeNode> SelectedAssetNode = AssetTableTreeView->GetSingleSelectedAssetNode();
+	if (SelectedAssetNode.IsValid())
+	{
+		check(!SelectedAssetNode->IsGroup());
+		check(SelectedAssetNode->IsValidAsset());
+		const FAssetTableRow& AssetTableRow = SelectedAssetNode->GetAssetChecked();
+		//...
+	}
+}
+
 void SAssetAuditBrowser::RefreshAssetView()
 {
 	FARFilter Filter;
@@ -804,175 +830,183 @@ void SAssetAuditBrowser::RefreshAssetView()
 
 	SetFilterDelegate.Execute(Filter);
 
+	RefreshAssetTableTreeView(Filter);
+}
 
-	// Refresh list of assets for the tree view
-	if (AssetTableTreeView.IsValid())
+// Refresh list of assets for the tree view
+void SAssetAuditBrowser::RefreshAssetTableTreeView(const FARFilter& Filter)
+{
+	if (!AssetTableTreeView.IsValid())
 	{
-		TSharedPtr<FAssetTable> AssetTable = AssetTableTreeView->GetAssetTable();
-		if (AssetTable.IsValid())
-		{
-			UE_LOG(LogInsights, Log, TEXT("[AssetTree] Build asset table..."));
+		return;
+	}
 
-			UE::Insights::FStopwatch Stopwatch;
-			Stopwatch.Start();
+	TSharedPtr<FAssetTable> AssetTable = AssetTableTreeView->GetAssetTable();
+	if (!AssetTable.IsValid())
+	{
+		return;
+	}
 
-			// Clears all tree nodes (that references the previous assets).
-			AssetTable->SetVisibleAssetCount(0);
-			AssetTableTreeView->RebuildTree(true);
+	UE_LOG(LogInsights, Log, TEXT("[AssetTree] Build asset table..."));
 
-			// Now is safe to clear the previous assets.
-			AssetTable->ClearAllData();
+	UE::Insights::FStopwatch Stopwatch;
+	Stopwatch.Start();
 
-			TMap<FAssetData, int32> AssetToIndexMap;
+	// Clears all tree nodes (that references the previous assets).
+	AssetTable->SetVisibleAssetCount(0);
+	AssetTableTreeView->RebuildTree(true);
 
-			TArray<FAssetData> SourceAssets;
-			const FAssetManagerEditorRegistrySource* RegistrySource = EditorModule->GetCurrentRegistrySource();
+	// Now is safe to clear the previous assets.
+	AssetTable->ClearAllData();
 
-			if (RegistrySource->GetOwnedRegistryState())
-			{
-				FARCompiledFilter CompiledFilter;
-				AssetRegistry->CompileFilter(Filter, CompiledFilter);
-				RegistrySource->GetOwnedRegistryState()->GetAssets(CompiledFilter, TSet<FName>(), SourceAssets);
-			}
-			else
+	TMap<FAssetData, int32> AssetToIndexMap;
+
+	TArray<FAssetData> SourceAssets;
+	const FAssetManagerEditorRegistrySource* RegistrySource = EditorModule->GetCurrentRegistrySource();
+
+	if (RegistrySource->GetOwnedRegistryState())
+	{
+		FARCompiledFilter CompiledFilter;
+		AssetRegistry->CompileFilter(Filter, CompiledFilter);
+		RegistrySource->GetOwnedRegistryState()->GetAssets(CompiledFilter, TSet<FName>(), SourceAssets);
+	}
+	else
+	{
+		AssetRegistry->GetAssets(Filter, SourceAssets);
+	}
+
+	for (int32 SourceAssetIndex = 0; SourceAssetIndex < SourceAssets.Num(); SourceAssetIndex++)
+	{
+		TArray<FAssetData> AssetsInSourcePackage;
+		AssetRegistry->GetAssetsByPackageName(SourceAssets[SourceAssetIndex].PackageName, AssetsInSourcePackage, /*bIncludeOnlyOnDiskAssets*/ true); // Only use on disk assets to avoid creating FAssetData for everything in memory
+
+		for (FAssetData& SourceAsset : AssetsInSourcePackage)
+		{
+			if (AssetToIndexMap.Find(SourceAsset) == nullptr)
 			{
-				AssetRegistry->GetAssets(Filter, SourceAssets);
+				FAssetTableRow AssetRow;
+				PopulateAssetTableRow(AssetRow, SourceAsset, *AssetTable);
+				AssetTable->AddAsset(AssetRow);
+				AssetToIndexMap.Add(SourceAsset, AssetTable->GetTotalAssetCount() - 1);
 			}
+		}
+	}
+	AssetTable->SetVisibleAssetCount(AssetTable->GetAssets().Num());
 
-			for (int32 SourceAssetIndex = 0; SourceAssetIndex < SourceAssets.Num(); SourceAssetIndex++)
-			{
-				TArray<FAssetData> AssetsInSourcePackage;
-				AssetRegistry->GetAssetsByPackageName(SourceAssets[SourceAssetIndex].PackageName, AssetsInSourcePackage, /*bIncludeOnlyOnDiskAssets*/ true); // Only use on disk assets to avoid creating FAssetData for everything in memory
+	int32 EntryIndex = AssetTable->GetAssets().Num();
 
-				for (FAssetData& SourceAsset : AssetsInSourcePackage)
-				{
-					if (AssetToIndexMap.Find(SourceAsset) == nullptr)
-					{
-						FAssetTableRow AssetRow;
-						PopulateAssetTableRow(AssetRow, SourceAsset, *AssetTable);
-						AssetTable->AddAsset(AssetRow);
-						AssetToIndexMap.Add(SourceAsset, AssetTable->GetTotalAssetCount() - 1);
-					}
-				}
-			}
-			AssetTable->SetVisibleAssetCount(AssetTable->GetAssets().Num());
+	TSet<FName> AlreadyProcessedPackages;
+	// Go over all SourceAssets. We will add more to the end as we find we need to add them.
+	for (int32 SourceAssetIndex = 0; SourceAssetIndex < SourceAssets.Num(); SourceAssetIndex++)
+	{
+		// There are potentially multiple assets in a package, each with an FAssetData
+		// We are only dealing in package dependencies, though, so don't process the same package twice
+		// But we will need to conceptually assign all the package dependencies to all assets in the package below
+		if (AlreadyProcessedPackages.Contains(SourceAssets[SourceAssetIndex].PackageName))
+		{
+			continue;
+		}
+		else
+		{
+			AlreadyProcessedPackages.Add(SourceAssets[SourceAssetIndex].PackageName);
+		}
 
-			int32 EntryIndex = AssetTable->GetAssets().Num();
+		FAssetIdentifier CurrentIdentifier(SourceAssets[SourceAssetIndex].PackageName);
+		// We'll have to add the dependencies to all these entries
+		TArray<FAssetData> AssetsInSourcePackage;
+		AssetRegistry->GetAssetsByPackageName(SourceAssets[SourceAssetIndex].PackageName, AssetsInSourcePackage, /*bIncludeOnlyOnDiskAssets*/ true); // Only use on disk assets to avoid creating FAssetData for everything in memory
 
-			TSet<FName> AlreadyProcessedPackages;
-			// Go over all SourceAssets. We will add more to the end as we find we need to add them.
-			for (int32 SourceAssetIndex = 0; SourceAssetIndex < SourceAssets.Num(); SourceAssetIndex++)
+		// Get the dependencies
+		TArray<FAssetIdentifier> DependencyList;
+		RegistrySource->GetDependencies(CurrentIdentifier, DependencyList, UE::AssetRegistry::EDependencyCategory::Package);
+
+		TArray<int32> IndicesOfDependenciesInRowTableToAddToCurrentSourceAssetRow;
+
+		for (const FAssetIdentifier& Dependency : DependencyList)
+		{
+			auto GatherIndices = [&](FAssetData const* DependencyAsset)
 			{
-				// There are potentially multiple assets in a package, each with an FAssetData
-				// We are only dealing in package dependencies, though, so don't process the same package twice
-				// But we will need to conceptually assign all the package dependencies to all assets in the package below
-				if (AlreadyProcessedPackages.Contains(SourceAssets[SourceAssetIndex].PackageName))
+				if (int32* DependencyIndex = AssetToIndexMap.Find(*DependencyAsset))
 				{
-					continue;
+					// Do we already know about the asset? If so, great.
+					IndicesOfDependenciesInRowTableToAddToCurrentSourceAssetRow.Add(*DependencyIndex);
 				}
 				else
 				{
-					AlreadyProcessedPackages.Add(SourceAssets[SourceAssetIndex].PackageName);
+					// We don't know about this asset yet. Create a new row for it and add it to the source list for further dependency analysis
+					FAssetTableRow NewRow;
+					PopulateAssetTableRow(NewRow, *DependencyAsset, *AssetTable);
+					AssetTable->AddAsset(NewRow);
+					IndicesOfDependenciesInRowTableToAddToCurrentSourceAssetRow.Add(AssetTable->GetTotalAssetCount() - 1);
+					AssetToIndexMap.Add(*DependencyAsset, AssetTable->GetTotalAssetCount() - 1);
+					SourceAssets.Add(*DependencyAsset);
 				}
+			};
 
-				FAssetIdentifier CurrentIdentifier(SourceAssets[SourceAssetIndex].PackageName);
-				// We'll have to add the dependencies to all these entries
-				TArray<FAssetData> AssetsInSourcePackage;
-				AssetRegistry->GetAssetsByPackageName(SourceAssets[SourceAssetIndex].PackageName, AssetsInSourcePackage, /*bIncludeOnlyOnDiskAssets*/ true); // Only use on disk assets to avoid creating FAssetData for everything in memory
-
-				// Get the dependencies
-				TArray<FAssetIdentifier> DependencyList;
-				RegistrySource->GetDependencies(CurrentIdentifier, DependencyList, UE::AssetRegistry::EDependencyCategory::Package);
-				
-				TArray<int32> IndicesOfDependenciesInRowTableToAddToCurrentSourceAssetRow;
+			if (RegistrySource->GetOwnedRegistryState())
+			{
+				TArrayView<FAssetData const* const> AssetsInDependencyPackage = RegistrySource->GetOwnedRegistryState()->GetAssetsByPackageName(Dependency.PackageName);
+				for (FAssetData const* const DependencyAsset : AssetsInDependencyPackage)
+				{
+					GatherIndices(DependencyAsset);
+				}
+			}
+			else
+			{
+				TArray<FAssetData> AssetsInDependencyPackage;
+				AssetRegistry->GetAssetsByPackageName(Dependency.PackageName, AssetsInDependencyPackage, /*only on disk assets*/true);
+				for (FAssetData& DependencyAsset : AssetsInDependencyPackage)
+				{
+					GatherIndices(&DependencyAsset);
+				}
+			}
+		}
 
-				for (const FAssetIdentifier& Dependency : DependencyList)
+		if (IndicesOfDependenciesInRowTableToAddToCurrentSourceAssetRow.Num())
+		{
+			// We found some dependencies. Let's add them.
+			for (const FAssetData& SourceAssetData : AssetsInSourcePackage)
+			{
+				int32* RowIndex = AssetToIndexMap.Find(SourceAssetData);
+				if (ensure(RowIndex != nullptr))
 				{
-					auto GatherIndices = [&](FAssetData const* DependencyAsset)
-					{
-						if (int32* DependencyIndex = AssetToIndexMap.Find(*DependencyAsset))
-						{
-							// Do we already know about the asset? If so, great.
-							IndicesOfDependenciesInRowTableToAddToCurrentSourceAssetRow.Add(*DependencyIndex);
-						}
-						else
-						{
-							// We don't know about this asset yet. Create a new row for it and add it to the source list for further dependency analysis
-							FAssetTableRow NewRow;
-							PopulateAssetTableRow(NewRow, *DependencyAsset, *AssetTable);
-							AssetTable->AddAsset(NewRow);
-							IndicesOfDependenciesInRowTableToAddToCurrentSourceAssetRow.Add(AssetTable->GetTotalAssetCount() - 1);
-							AssetToIndexMap.Add(*DependencyAsset, AssetTable->GetTotalAssetCount() - 1);
-							SourceAssets.Add(*DependencyAsset);
-						}
-					};
-					
-					if (RegistrySource->GetOwnedRegistryState())
-					{
-						TArrayView<FAssetData const* const> AssetsInDependencyPackage = RegistrySource->GetOwnedRegistryState()->GetAssetsByPackageName(Dependency.PackageName);
-						for (FAssetData const* const DependencyAsset : AssetsInDependencyPackage)
-						{
-							GatherIndices(DependencyAsset);
-						}
-					}
-					else
+					for (int32 Index : IndicesOfDependenciesInRowTableToAddToCurrentSourceAssetRow)
 					{
-						TArray<FAssetData> AssetsInDependencyPackage;
-						AssetRegistry->GetAssetsByPackageName(Dependency.PackageName, AssetsInDependencyPackage, /*only on disk assets*/true);
-						for (FAssetData& DependencyAsset : AssetsInDependencyPackage)
-						{
-							GatherIndices(&DependencyAsset);
-						}
+						AssetTable->GetAsset(*RowIndex)->Dependencies.AddUnique(Index);
 					}
 				}
+			}
 
-				if (IndicesOfDependenciesInRowTableToAddToCurrentSourceAssetRow.Num())
+			for (int32 Index : IndicesOfDependenciesInRowTableToAddToCurrentSourceAssetRow)
+			{
+				// Now for each of those dependencies, add this source asset row as a referencer
+				if (FAssetTableRow* DependentRow = AssetTable->GetAsset(Index))
 				{
-					// We found some dependencies. Let's add them.
 					for (const FAssetData& SourceAssetData : AssetsInSourcePackage)
 					{
 						int32* RowIndex = AssetToIndexMap.Find(SourceAssetData);
 						if (ensure(RowIndex != nullptr))
 						{
-							for (int32 Index : IndicesOfDependenciesInRowTableToAddToCurrentSourceAssetRow)
-							{
-								AssetTable->GetAsset(*RowIndex)->Dependencies.AddUnique(Index);
-							}
-						}
-					}
-
-					for (int32 Index : IndicesOfDependenciesInRowTableToAddToCurrentSourceAssetRow)
-					{
-						// Now for each of those dependencies, add this source asset row as a referencer
-						if (FAssetTableRow* DependentRow = AssetTable->GetAsset(Index))
-						{
-							for (const FAssetData& SourceAssetData : AssetsInSourcePackage)
-							{
-								int32* RowIndex = AssetToIndexMap.Find(SourceAssetData);
-								if (ensure(RowIndex != nullptr))
-								{
-									DependentRow->Referencers.AddUnique(*RowIndex);
-								}
-							}
+							DependentRow->Referencers.AddUnique(*RowIndex);
 						}
 					}
 				}
 			}
+		}
+	}
 
-			Stopwatch.Stop();
-			const double TotalTime = Stopwatch.GetAccumulatedTime();
-			UE_LOG(LogInsights, Log, TEXT("[AssetTree] Asset table rebuilt in %.4fs (%d visible + %d hidden = %d assets)"),
-				TotalTime, AssetTable->GetVisibleAssetCount(), AssetTable->GetHiddenAssetCount(), AssetTable->GetTotalAssetCount());
+	Stopwatch.Stop();
+	const double TotalTime = Stopwatch.GetAccumulatedTime();
+	UE_LOG(LogInsights, Log, TEXT("[AssetTree] Asset table rebuilt in %.4fs (%d visible + %d hidden = %d assets)"),
+		TotalTime, AssetTable->GetVisibleAssetCount(), AssetTable->GetHiddenAssetCount(), AssetTable->GetTotalAssetCount());
 
-			const FAssetTableStringStore& StringStore = AssetTable->GetStringStore();
-			UE_LOG(LogInsights, Log, TEXT("[AssetTree] String store: %d strings (%lld bytes) --> %d strings (%lld bytes | %lld bytes) %.0f%%"),
-				StringStore.GetNumInputStrings(), StringStore.GetTotalInputStringSize(), 
-				StringStore.GetNumStrings(), StringStore.GetTotalStringSize(), StringStore.GetAllocatedSize(),
-				(double)StringStore.GetAllocatedSize() * 100.0 / (double)StringStore.GetTotalInputStringSize());
+	const FAssetTableStringStore& StringStore = AssetTable->GetStringStore();
+	UE_LOG(LogInsights, Log, TEXT("[AssetTree] String store: %d strings (%lld bytes) --> %d strings (%lld bytes | %lld bytes) %.0f%%"),
+		StringStore.GetNumInputStrings(), StringStore.GetTotalInputStringSize(),
+		StringStore.GetNumStrings(), StringStore.GetTotalStringSize(), StringStore.GetAllocatedSize(),
+		(double)StringStore.GetAllocatedSize() * 100.0 / (double)StringStore.GetTotalInputStringSize());
 
-			AssetTableTreeView->RebuildTreeAsync();
-		}
-	}
+	AssetTableTreeView->RebuildTreeAsync();
 }
 
 FReply SAssetAuditBrowser::ClearAssets()
@@ -1005,7 +1039,7 @@ void SAssetAuditBrowser::AddAssetsOfType(FPrimaryAssetType AssetType)
 	if (AssetType.IsValid())
 	{
 		TArray<FSoftObjectPath> AssetArray;
-		
+
 		if (AssetType == IAssetManagerEditorModule::AllPrimaryAssetTypes)
 		{
 			TArray<FPrimaryAssetTypeInfo> TypeList;
@@ -1018,7 +1052,7 @@ void SAssetAuditBrowser::AddAssetsOfType(FPrimaryAssetType AssetType)
 		}
 		else
 		{
-		AssetManager->GetPrimaryAssetPathList(AssetType, AssetArray);
+			AssetManager->GetPrimaryAssetPathList(AssetType, AssetArray);
 		}
 
 		AddAssetsToList(AssetArray, false);
@@ -1057,7 +1091,7 @@ void SAssetAuditBrowser::AddAssetsOfClass(UClass* AssetClass)
 			AssetRegistry->GetDerivedClassNames(AssetFilter.ClassPaths, TSet<FTopLevelAssetPath>(), DerivedClassNames);
 			AssetFilter.ClassPaths.Add(UBlueprintCore::StaticClass()->GetClassPathName());
 		}
-		
+
 		if (AssetRegistry->GetAssets(AssetFilter, FoundData) && FoundData.Num() > 0)
 		{
 			TArray<FName> AssetPackageArray;
@@ -1108,7 +1142,42 @@ bool SAssetAuditBrowser::HandleFilterAsset(const FAssetData& InAssetData) const
 	}
 
 	return !EditorModule->IsPackageInCurrentRegistrySource(InAssetData.PackageName);
+}
+
+void SAssetAuditBrowser::OnAssetSelected(const FAssetData& InAssetData)
+{
+	// Find and select the corresponding asset in the asset tree view.
+	if (AssetTableTreeView.IsValid())
+	{
+		TSharedPtr<FAssetTable> AssetTable = AssetTableTreeView->GetAssetTable();
+		if (AssetTable.IsValid())
+		{
+			FAssetTableRow SelectedAsset;
+			QuickPopulateAssetTableRow(SelectedAsset, InAssetData, *AssetTable);
+
+			int32 RowIndex = -1;
+
+			const TArray<FAssetTableRow>& Assets = AssetTable->GetAssets();
+			int32 AssetCount = AssetTable->GetVisibleAssetCount();
+			for (int32 Index = 0; Index < AssetCount; ++Index)
+			{
+				const FAssetTableRow& Asset = Assets[Index];
+				if (Asset.GetType() == SelectedAsset.GetType() &&
+					Asset.GetPath() == SelectedAsset.GetPath() &&
+					Asset.GetName() == SelectedAsset.GetName())
+				{
+					RowIndex = Index;
+					break;
+				}
+			}
+
+			if (RowIndex >= 0)
+			{
+				AssetTableTreeView->SelectNodeByTableRowIndex(RowIndex);
+			}
+		}
 	}
+}
 
 TSharedRef<SWidget> SAssetAuditBrowser::GenerateSourceComboItem(TSharedPtr<FString> InItem)
 {
@@ -1144,7 +1213,6 @@ void SAssetAuditBrowser::SetCurrentRegistrySource(const FAssetManagerEditorRegis
 	}
 	else
 	{
-		
 		if (CurrentRegistrySource->SourceTimestamp.Len() == 0)
 		{
 			RegistrySourceTimeText->SetText(LOCTEXT("AssetRegistryTimestampError", "Unable to get AssetRegistry.bin timestamp (recook?)"));
diff --git a/Engine/Plugins/Editor/AssetManagerEditor/Source/AssetManagerEditor/Private/SAssetAuditBrowser.h b/Engine/Plugins/Editor/AssetManagerEditor/Source/AssetManagerEditor/Private/SAssetAuditBrowser.h
--- a/Engine/Plugins/Editor/AssetManagerEditor/Source/AssetManagerEditor/Private/SAssetAuditBrowser.h
+++ b/Engine/Plugins/Editor/AssetManagerEditor/Source/AssetManagerEditor/Private/SAssetAuditBrowser.h
@@ -12,6 +12,14 @@ class FUICommandList;
 class SMenuAnchor;
 class STextBlock;
 
+namespace UE
+{
+	namespace Insights
+	{
+		class FTableTreeNode;
+	}
+}
+
 //////////////////////////////////////////////////////////////////////////
 // SAssetAuditBrowser
 
@@ -105,6 +113,8 @@ protected:
 	/** Perform additional filtering */
 	bool HandleFilterAsset(const FAssetData& InAssetData) const;
 
+	void OnAssetSelected(const FAssetData& InAssetData);
+
 	/** Button callbacks */
 	FReply ClearAssets();
 	FReply AddChunks();
@@ -130,9 +140,16 @@ protected:
 	/** Refresh the asset view with a new filter */
 	void RefreshAssetView();
 
+	void RefreshAssetTableTreeView(const FARFilter& Filter);
+
+	/** Populates an asset table row (view model for the asset tree view) with asset data (only identification data; type+path+name). */
+	void QuickPopulateAssetTableRow(class FAssetTableRow& OutRow, const FAssetData& AssetData, class FAssetTable& AssetTable) const;
+
 	/** Populates an asset table row (view model for the asset tree view) with asset data. */
 	void PopulateAssetTableRow(class FAssetTableRow& OutRow, const FAssetData& AssetData, class FAssetTable& AssetTable) const;
 
+	void TreeView_OnSelectionChanged(const TArray<TSharedPtr<UE::Insights::FTableTreeNode>> InSelectedNodes);
+
 	/** Commands handled by this widget */
 	TSharedPtr<FUICommandList> Commands;
 
diff --git a/Engine/Plugins/Editor/AssetManagerEditor/Source/AssetManagerEditor/Private/TreeView/AssetTable.cpp b/Engine/Plugins/Editor/AssetManagerEditor/Source/AssetManagerEditor/Private/TreeView/AssetTable.cpp
--- a/Engine/Plugins/Editor/AssetManagerEditor/Source/AssetManagerEditor/Private/TreeView/AssetTable.cpp
+++ b/Engine/Plugins/Editor/AssetManagerEditor/Source/AssetManagerEditor/Private/TreeView/AssetTable.cpp
@@ -9,6 +9,8 @@
 
 // AssetManagerEditor
 #include "AssetTreeNode.h"
+#include "Insights/Common/Log.h"
+#include "Insights/Common/Stopwatch.h"
 #include "Insights/Table/ViewModels/TableCellValueFormatter.h"
 #include "Insights/Table/ViewModels/TableCellValueGetter.h"
 #include "Insights/Table/ViewModels/TableCellValueSorter.h"
@@ -20,8 +22,8 @@
 // Column identifiers
 
 const FName FAssetTableColumns::CountColumnId(TEXT("Count"));
-const FName FAssetTableColumns::NameColumnId(TEXT("Name"));
 const FName FAssetTableColumns::TypeColumnId(TEXT("Type"));
+const FName FAssetTableColumns::NameColumnId(TEXT("Name"));
 const FName FAssetTableColumns::PathColumnId(TEXT("Path"));
 const FName FAssetTableColumns::PrimaryTypeColumnId(TEXT("PrimaryType"));
 const FName FAssetTableColumns::PrimaryNameColumnId(TEXT("PrimaryName"));
@@ -226,7 +228,7 @@ FAssetTable::FAssetTable()
 			UniqueDependencies[2] = TSet<int32>{3};
 			UniqueDependencies[3] = TSet<int32>{};
 			UniqueDependencies[4] = TSet<int32>{0, 1, 2, 3};
-   
+
 			SharedDependencies[0] = TSet<int32>{};
 			SharedDependencies[1] = TSet<int32>{};
 			SharedDependencies[2] = TSet<int32>{};
@@ -287,7 +289,7 @@ FAssetTable::FAssetTable()
 			UniqueDependencies[2] = TSet<int32>{ 3, 0, 1 };
 			UniqueDependencies[3] = TSet<int32>{ 0, 1 };
 
-			// Unlike the previous test, since 5 also references 2, 
+			// Unlike the previous test, since 5 also references 2,
 			// 4 will see itself as having no unique dependencies
 			UniqueDependencies[4] = TSet<int32>{};
 			UniqueDependencies[5] = TSet<int32>{};
@@ -377,7 +379,7 @@ FAssetTable::FAssetTable()
 	VisibleAssetCount = 100;
 	constexpr int32 HiddenAssetCount = 50;
 	const int32 TotalAssetCount = VisibleAssetCount + HiddenAssetCount;
-	
+
 	// Create assets.
 	Assets.Reserve(TotalAssetCount);
 	for (int32 AssetIndex = 0; AssetIndex < TotalAssetCount; ++AssetIndex)
@@ -465,7 +467,7 @@ void FAssetTable::AddDefaultColumns()
 		AddHierarchyColumn(HierarchyColumnIndex, HierarchyColumnName);
 
 		const TSharedRef<FTableColumn>& ColumnRef = GetColumns()[0];
-		ColumnRef->SetInitialWidth(200.0f);
+		ColumnRef->SetInitialWidth(400.0f);
 		ColumnRef->SetShortName(LOCTEXT("HierarchyColumnName", "Hierarchy"));
 		ColumnRef->SetTitleName(LOCTEXT("HierarchyColumnTitle", "Asset Hierarchy"));
 		ColumnRef->SetDescription(LOCTEXT("HierarchyColumnDesc", "Hierarchy of the asset tree"));
@@ -485,7 +487,7 @@ void FAssetTable::AddDefaultColumns()
 		Column.SetTitleName(LOCTEXT("CountColumnTitle", "Asset Count"));
 		Column.SetDescription(LOCTEXT("CountColumnDesc", "Number of assets"));
 
-		Column.SetFlags(ETableColumnFlags::ShouldBeVisible | ETableColumnFlags::CanBeHidden | ETableColumnFlags::CanBeFiltered);
+		Column.SetFlags(ETableColumnFlags::ShouldBeVisible | ETableColumnFlags::CanBeHidden | ETableColumnFlags::CanBeFiltered | ETableColumnFlags::IsDynamic);
 
 		Column.SetHorizontalAlignment(HAlign_Right);
 		Column.SetInitialWidth(100.0f);
@@ -530,25 +532,25 @@ void FAssetTable::AddDefaultColumns()
 		AddColumn(ColumnRef);
 	}
 	//////////////////////////////////////////////////
-	// Name Column
+	// Type Column
 	{
-		TSharedRef<FTableColumn> ColumnRef = MakeShared<FTableColumn>(FAssetTableColumns::NameColumnId);
+		TSharedRef<FTableColumn> ColumnRef = MakeShared<FTableColumn>(FAssetTableColumns::TypeColumnId);
 		FTableColumn& Column = *ColumnRef;
 
 		Column.SetIndex(ColumnIndex++);
 
-		Column.SetShortName(LOCTEXT("NameColumnName", "Name"));
-		Column.SetTitleName(LOCTEXT("NameColumnTitle", "Name"));
-		Column.SetDescription(LOCTEXT("NameColumnDesc", "Asset's name"));
+		Column.SetShortName(LOCTEXT("TypeColumnName", "Type"));
+		Column.SetTitleName(LOCTEXT("TypeColumnTitle", "Type"));
+		Column.SetDescription(LOCTEXT("TypeColumnDesc", "Asset's type"));
 
-		Column.SetFlags(ETableColumnFlags::ShouldBeVisible | ETableColumnFlags::CanBeHidden | ETableColumnFlags::CanBeFiltered);
+		Column.SetFlags(ETableColumnFlags::CanBeHidden | ETableColumnFlags::CanBeFiltered);
 
 		Column.SetHorizontalAlignment(HAlign_Left);
-		Column.SetInitialWidth(120.0f);
+		Column.SetInitialWidth(200.0f);
 
 		Column.SetDataType(ETableCellDataType::CString);
 
-		class FAssetNameValueGetter : public FTableCellValueGetter
+		class FAssetTypeValueGetter : public FTableCellValueGetter
 		{
 		public:
 			virtual const TOptional<FTableCellValue> GetValue(const FTableColumn& Column, const FBaseTreeNode& Node) const
@@ -565,13 +567,13 @@ void FAssetTable::AddDefaultColumns()
 				{
 					const FAssetTreeNode& TreeNode = static_cast<const FAssetTreeNode&>(Node);
 					const FAssetTableRow& Asset = TreeNode.GetAssetChecked();
-					return FTableCellValue(Asset.GetName());
+					return FTableCellValue(Asset.GetType());
 				}
 
 				return TOptional<FTableCellValue>();
 			}
 		};
-		TSharedRef<ITableCellValueGetter> Getter = MakeShared<FAssetNameValueGetter>();
+		TSharedRef<ITableCellValueGetter> Getter = MakeShared<FAssetTypeValueGetter>();
 		Column.SetValueGetter(Getter);
 
 		TSharedRef<ITableCellValueFormatter> Formatter = MakeShared<FCStringValueFormatterAsText>();
@@ -585,25 +587,25 @@ void FAssetTable::AddDefaultColumns()
 		AddColumn(ColumnRef);
 	}
 	//////////////////////////////////////////////////
-	// Type Column
+	// Name Column
 	{
-		TSharedRef<FTableColumn> ColumnRef = MakeShared<FTableColumn>(FAssetTableColumns::TypeColumnId);
+		TSharedRef<FTableColumn> ColumnRef = MakeShared<FTableColumn>(FAssetTableColumns::NameColumnId);
 		FTableColumn& Column = *ColumnRef;
 
 		Column.SetIndex(ColumnIndex++);
 
-		Column.SetShortName(LOCTEXT("TypeColumnName", "Type"));
-		Column.SetTitleName(LOCTEXT("TypeColumnTitle", "Type"));
-		Column.SetDescription(LOCTEXT("TypeColumnDesc", "Asset's type"));
+		Column.SetShortName(LOCTEXT("NameColumnName", "Name"));
+		Column.SetTitleName(LOCTEXT("NameColumnTitle", "Name"));
+		Column.SetDescription(LOCTEXT("NameColumnDesc", "Asset's name"));
 
-		Column.SetFlags(ETableColumnFlags::ShouldBeVisible | ETableColumnFlags::CanBeHidden | ETableColumnFlags::CanBeFiltered);
+		Column.SetFlags(ETableColumnFlags::CanBeHidden | ETableColumnFlags::CanBeFiltered);
 
 		Column.SetHorizontalAlignment(HAlign_Left);
-		Column.SetInitialWidth(120.0f);
+		Column.SetInitialWidth(200.0f);
 
 		Column.SetDataType(ETableCellDataType::CString);
 
-		class FAssetTypeValueGetter : public FTableCellValueGetter
+		class FAssetNameValueGetter : public FTableCellValueGetter
 		{
 		public:
 			virtual const TOptional<FTableCellValue> GetValue(const FTableColumn& Column, const FBaseTreeNode& Node) const
@@ -620,13 +622,13 @@ void FAssetTable::AddDefaultColumns()
 				{
 					const FAssetTreeNode& TreeNode = static_cast<const FAssetTreeNode&>(Node);
 					const FAssetTableRow& Asset = TreeNode.GetAssetChecked();
-					return FTableCellValue(Asset.GetType());
+					return FTableCellValue(Asset.GetName());
 				}
 
 				return TOptional<FTableCellValue>();
 			}
 		};
-		TSharedRef<ITableCellValueGetter> Getter = MakeShared<FAssetTypeValueGetter>();
+		TSharedRef<ITableCellValueGetter> Getter = MakeShared<FAssetNameValueGetter>();
 		Column.SetValueGetter(Getter);
 
 		TSharedRef<ITableCellValueFormatter> Formatter = MakeShared<FCStringValueFormatterAsText>();
@@ -651,10 +653,10 @@ void FAssetTable::AddDefaultColumns()
 		Column.SetTitleName(LOCTEXT("PathColumnTitle", "Path"));
 		Column.SetDescription(LOCTEXT("PathColumnDesc", "Asset's path"));
 
-		Column.SetFlags(ETableColumnFlags::ShouldBeVisible | ETableColumnFlags::CanBeHidden | ETableColumnFlags::CanBeFiltered);
+		Column.SetFlags(ETableColumnFlags::CanBeHidden | ETableColumnFlags::CanBeFiltered);
 
 		Column.SetHorizontalAlignment(HAlign_Left);
-		Column.SetInitialWidth(120.0f);
+		Column.SetInitialWidth(200.0f);
 
 		Column.SetDataType(ETableCellDataType::CString);
 
@@ -709,7 +711,7 @@ void FAssetTable::AddDefaultColumns()
 		Column.SetFlags(ETableColumnFlags::ShouldBeVisible | ETableColumnFlags::CanBeHidden | ETableColumnFlags::CanBeFiltered);
 
 		Column.SetHorizontalAlignment(HAlign_Left);
-		Column.SetInitialWidth(120.0f);
+		Column.SetInitialWidth(200.0f);
 
 		Column.SetDataType(ETableCellDataType::CString);
 
@@ -764,7 +766,7 @@ void FAssetTable::AddDefaultColumns()
 		Column.SetFlags(ETableColumnFlags::ShouldBeVisible | ETableColumnFlags::CanBeHidden | ETableColumnFlags::CanBeFiltered);
 
 		Column.SetHorizontalAlignment(HAlign_Left);
-		Column.SetInitialWidth(120.0f);
+		Column.SetInitialWidth(200.0f);
 
 		Column.SetDataType(ETableCellDataType::CString);
 
@@ -812,11 +814,11 @@ void FAssetTable::AddDefaultColumns()
 
 		Column.SetIndex(ColumnIndex++);
 
-		Column.SetShortName(LOCTEXT("StagedCompressedSizeColumnName", "Staged Compressed Size"));
+		Column.SetShortName(LOCTEXT("StagedCompressedSizeColumnName", "Staged"));
 		Column.SetTitleName(LOCTEXT("StagedCompressedSizeColumnTitle", "Staged Compressed Size"));
 		Column.SetDescription(LOCTEXT("StagedCompressedSizeColumnDesc", "Compressed size of iostore chunks for this asset's package. Only visible after staging."));
 
-		Column.SetFlags(ETableColumnFlags::ShouldBeVisible | ETableColumnFlags::CanBeHidden | ETableColumnFlags::CanBeFiltered);
+		Column.SetFlags(ETableColumnFlags::ShouldBeVisible | ETableColumnFlags::CanBeHidden | ETableColumnFlags::CanBeFiltered | ETableColumnFlags::IsDynamic);
 
 		Column.SetHorizontalAlignment(HAlign_Right);
 		Column.SetInitialWidth(100.0f);
@@ -871,14 +873,14 @@ void FAssetTable::AddDefaultColumns()
 
 		Column.SetIndex(ColumnIndex++);
 
-		Column.SetShortName(LOCTEXT("TotalSizeUniqueDependenciesColumnName", "Total Unique Dependency Size"));
+		Column.SetShortName(LOCTEXT("TotalSizeUniqueDependenciesColumnName", "Unique"));
 		Column.SetTitleName(LOCTEXT("TotalSizeUniqueDependenciesColumnTitle", "Total Unique Dependency Size"));
 		Column.SetDescription(LOCTEXT("TotalSizeUniqueDependenciesColumnIdDesc", "Sum of the staged compressed sizes of all dependencies of this asset, counted only once"));
 
-		Column.SetFlags(ETableColumnFlags::ShouldBeVisible | ETableColumnFlags::CanBeHidden | ETableColumnFlags::CanBeFiltered);
+		Column.SetFlags(ETableColumnFlags::CanBeHidden | ETableColumnFlags::CanBeFiltered);
 
 		Column.SetHorizontalAlignment(HAlign_Right);
-		Column.SetInitialWidth(120.0f);
+		Column.SetInitialWidth(100.0f);
 
 		Column.SetDataType(ETableCellDataType::Int64);
 
@@ -899,7 +901,7 @@ void FAssetTable::AddDefaultColumns()
 				{
 					const FAssetTreeNode& TreeNode = static_cast<const FAssetTreeNode&>(Node);
 					const FAssetTableRow& Asset = TreeNode.GetAssetChecked();
-					return FTableCellValue(static_cast<int64>(Asset.GetOrComputeTotalSizeUniqueDependencies(&TreeNode.GetAssetTableChecked(), TreeNode.GetRowIndex())));
+					return FTableCellValue(static_cast<int64>(Asset.GetOrComputeTotalSizeUniqueDependencies(TreeNode.GetAssetTableChecked(), TreeNode.GetRowIndex())));
 				}
 
 				return TOptional<FTableCellValue>();
@@ -930,14 +932,14 @@ void FAssetTable::AddDefaultColumns()
 
 		Column.SetIndex(ColumnIndex++);
 
-		Column.SetShortName(LOCTEXT("TotalSizeOtherDependenciesColumnName", "Total Other Dependency Size"));
+		Column.SetShortName(LOCTEXT("TotalSizeOtherDependenciesColumnName", "Other"));
 		Column.SetTitleName(LOCTEXT("TotalSizeOtherDependenciesColumnTitle", "Total Other Dependency Size"));
 		Column.SetDescription(LOCTEXT("TotalSizeOtherDependenciesColumnIdDesc", "Sum of the staged compressed sizes of all dependencies of this asset which are shared by other assets directly or indirectly, counted only once"));
 
-		Column.SetFlags(ETableColumnFlags::ShouldBeVisible | ETableColumnFlags::CanBeHidden | ETableColumnFlags::CanBeFiltered);
+		Column.SetFlags(ETableColumnFlags::CanBeHidden | ETableColumnFlags::CanBeFiltered);
 
 		Column.SetHorizontalAlignment(HAlign_Right);
-		Column.SetInitialWidth(120.0f);
+		Column.SetInitialWidth(100.0f);
 
 		Column.SetDataType(ETableCellDataType::Int64);
 
@@ -958,7 +960,7 @@ void FAssetTable::AddDefaultColumns()
 				{
 					const FAssetTreeNode& TreeNode = static_cast<const FAssetTreeNode&>(Node);
 					const FAssetTableRow& Asset = TreeNode.GetAssetChecked();
-					return FTableCellValue(static_cast<int64>(Asset.GetOrComputeTotalSizeOtherDependencies(&TreeNode.GetAssetTableChecked(), TreeNode.GetRowIndex())));
+					return FTableCellValue(static_cast<int64>(Asset.GetOrComputeTotalSizeOtherDependencies(TreeNode.GetAssetTableChecked(), TreeNode.GetRowIndex())));
 				}
 
 				return TOptional<FTableCellValue>();
@@ -993,7 +995,7 @@ void FAssetTable::AddDefaultColumns()
 		Column.SetTitleName(LOCTEXT("TotalUsageCountColumnTitle", "Total Usage Count"));
 		Column.SetDescription(LOCTEXT("TotalUsageCountColumnDesc", "Weighted count of Primary Assets that use this\nA higher usage means it's more likely to be in memory at runtime."));
 
-		Column.SetFlags(ETableColumnFlags::ShouldBeVisible | ETableColumnFlags::CanBeHidden | ETableColumnFlags::CanBeFiltered);
+		Column.SetFlags(ETableColumnFlags::ShouldBeVisible | ETableColumnFlags::CanBeHidden | ETableColumnFlags::CanBeFiltered | ETableColumnFlags::IsDynamic);
 
 		Column.SetHorizontalAlignment(HAlign_Right);
 		Column.SetInitialWidth(100.0f);
@@ -1052,7 +1054,7 @@ void FAssetTable::AddDefaultColumns()
 		Column.SetFlags(ETableColumnFlags::ShouldBeVisible | ETableColumnFlags::CanBeHidden | ETableColumnFlags::CanBeFiltered);
 
 		Column.SetHorizontalAlignment(HAlign_Left);
-		Column.SetInitialWidth(120.0f);
+		Column.SetInitialWidth(200.0f);
 
 		Column.SetDataType(ETableCellDataType::CString);
 
@@ -1107,7 +1109,7 @@ void FAssetTable::AddDefaultColumns()
 		Column.SetFlags(ETableColumnFlags::ShouldBeVisible | ETableColumnFlags::CanBeHidden | ETableColumnFlags::CanBeFiltered);
 
 		Column.SetHorizontalAlignment(HAlign_Left);
-		Column.SetInitialWidth(120.0f);
+		Column.SetInitialWidth(200.0f);
 
 		Column.SetDataType(ETableCellDataType::CString);
 
@@ -1162,7 +1164,7 @@ void FAssetTable::AddDefaultColumns()
 		Column.SetFlags(ETableColumnFlags::ShouldBeVisible | ETableColumnFlags::CanBeHidden | ETableColumnFlags::CanBeFiltered);
 
 		Column.SetHorizontalAlignment(HAlign_Left);
-		Column.SetInitialWidth(120.0f);
+		Column.SetInitialWidth(200.0f);
 
 		Column.SetDataType(ETableCellDataType::CString);
 
@@ -1205,184 +1207,197 @@ void FAssetTable::AddDefaultColumns()
 	//////////////////////////////////////////////////
 }
 
-// Iteratively refines the dependency set into unique and shared sets. PreviouslyVisitedIndices is split in each pass, moving some elements into 
-// OutExcludedIndices (preserving its original contents) and putting those that are still potentially uniquely owned by the ThisIndex asset into 
+////////////////////////////////////////////////////////////////////////////////////////////////////
+// Iteratively refines the dependency set into unique and shared sets. PreviouslyVisitedIndices is split in each pass, moving some elements into
+// OutExcludedIndices (preserving its original contents) and putting those that are still potentially uniquely owned by the ThisIndex asset into
 // OutIncrementallyRefinedUniqueIndices. Returns 'true' if another pass is required (i.e., if work was done).
-/*static*/ bool FAssetTableRow::RefineDependencies(TSet<int32> PreviouslyVisitedIndices, FAssetTable* OwningTable, int32 ThisIndex, TSet<int32>* OutIncrementallyRefinedUniqueIndices, TSet<int32>* OutExcludedIndices)
+/*static*/ bool FAssetTableRow::RefineDependencies(const TSet<int32>& PreviouslyVisitedIndices, FAssetTable& OwningTable, int32 ThisIndex, TSet<int32>& OutIncrementallyRefinedUniqueIndices, TSet<int32>& OutExcludedIndices)
 {
-	if (ensure(OutIncrementallyRefinedUniqueIndices != nullptr) && ensure(OutExcludedIndices != nullptr))
-	{
-		const FString& ThisGFP = OwningTable->GetAsset(ThisIndex)->GameFeaturePlugin;
+	const TCHAR* ThisGFP = OwningTable.GetAssetChecked(ThisIndex).GetGameFeaturePlugin();
 
-		// "visit" ThisIndex to seed the exploration
-		TArray<int32> IndicesToVisit = OwningTable->GetAsset(ThisIndex)->GetDependencies();
+	// "visit" ThisIndex to seed the exploration
+	TArray<int32> IndicesToVisit = OwningTable.GetAssetChecked(ThisIndex).GetDependencies();
 
-		while (IndicesToVisit.Num() > 0)
-		{
-			int32 CurrentIndex = IndicesToVisit.Pop(false);
+	while (IndicesToVisit.Num() > 0)
+	{
+		int32 CurrentIndex = IndicesToVisit.Pop(false);
 
-			FAssetTableRow* Row = OwningTable->GetAsset(CurrentIndex);
-			if (Row->GameFeaturePlugin != ThisGFP)
-			{
-				// Don't traverse outside this plugin
-				continue;
-			}
+		FAssetTableRow& Row = OwningTable.GetAssetChecked(CurrentIndex);
+		if (Row.GetGameFeaturePlugin() != ThisGFP)
+		{
+			// Don't traverse outside this plugin
+			continue;
+		}
 
-			bool ShouldIncludeInTotal = true;
-			for (int32 ReferencerIndex : Row->Referencers)
-			{
-				if (PreviouslyVisitedIndices.Contains(ReferencerIndex) == false && ThisIndex != ReferencerIndex)
-				{
-					ShouldIncludeInTotal = false;
-					break;
-				}
-			}
-			if (ShouldIncludeInTotal == false)
+		bool ShouldIncludeInTotal = true;
+		for (int32 ReferencerIndex : Row.Referencers)
+		{
+			if (PreviouslyVisitedIndices.Contains(ReferencerIndex) == false && ThisIndex != ReferencerIndex)
 			{
-				OutExcludedIndices->Add(CurrentIndex);
-				continue;
+				ShouldIncludeInTotal = false;
+				break;
 			}
-			OutIncrementallyRefinedUniqueIndices->Add(CurrentIndex);
+		}
+		if (ShouldIncludeInTotal == false)
+		{
+			OutExcludedIndices.Add(CurrentIndex);
+			continue;
+		}
+		OutIncrementallyRefinedUniqueIndices.Add(CurrentIndex);
 
-			for (int32 ChildIndex : OwningTable->GetAsset(CurrentIndex)->GetDependencies())
+		for (int32 ChildIndex : OwningTable.GetAssetChecked(CurrentIndex).GetDependencies())
+		{
+			// Don't revisit nodes we've already visited and don't re-add ThisIndex to avoid loops (and to avoid counting ourself)
+			if (!OutIncrementallyRefinedUniqueIndices.Contains(ChildIndex) && ThisIndex != ChildIndex)
 			{
-				// Don't revisit nodes we've already visited and don't re-add ThisIndex to avoid loops (and to avoid counting ourself)
-				if (!OutIncrementallyRefinedUniqueIndices->Contains(ChildIndex) && ThisIndex != ChildIndex)
-				{
-					IndicesToVisit.AddUnique(ChildIndex);
-				}
+				IndicesToVisit.AddUnique(ChildIndex);
 			}
 		}
-
-		return OutIncrementallyRefinedUniqueIndices->Num() != PreviouslyVisitedIndices.Num();
-	}
-	else
-	{
-		return false;
 	}
+
+	return OutIncrementallyRefinedUniqueIndices.Num() != PreviouslyVisitedIndices.Num();
 }
 
-void FAssetTableRow::ComputeDependencySizes(FAssetTable* OwningTable, int32 ThisIndex, TSet<int32>* OutUniqueDependencies/* = nullptr*/, TSet<int32>* OutSharedDependencies/* = nullptr*/) const
+////////////////////////////////////////////////////////////////////////////////////////////////////
+
+void FAssetTableRow::ComputeDependencySizes(FAssetTable& OwningTable, int32 ThisIndex, TSet<int32>* OutUniqueDependencies, TSet<int32>* OutSharedDependencies) const
 {
-	if (OutUniqueDependencies != nullptr && OutSharedDependencies != nullptr)
+	TotalSizeUniqueDependencies = 0;
+	TotalSizeOtherDependencies = 0;
+
+	TArray<int32> IndicesToVisit;
+	TSet<int32> VisitedIndices;
+
+	// Break any loops in the dependency graph
+	VisitedIndices.Add(ThisIndex);
+
+	const TCHAR* ThisGFP = OwningTable.GetAssetChecked(ThisIndex).GetGameFeaturePlugin();
+
+	// Don't include this asset itself in the total, just its children
+	for (int32 ChildIndex : Dependencies)
 	{
-		TotalSizeUniqueDependencies = -1;
+		IndicesToVisit.Add(ChildIndex);
 	}
-	if (TotalSizeUniqueDependencies == -1)
-	{
-		TArray<int32> IndicesToVisit;
-		TSet<int32> VisitedIndices;
-
-		TotalSizeUniqueDependencies = 0;
-		TotalSizeOtherDependencies = 0;
-		
-		// Break any loops in the dependency graph
-		VisitedIndices.Add(ThisIndex);
 
-		const FString& ThisGFP = OwningTable->GetAsset(ThisIndex)->GameFeaturePlugin;
+	while (IndicesToVisit.Num() > 0)
+	{
+		int32 CurrentIndex = IndicesToVisit.Pop(false);
 
-		// Don't include this asset itself in the total, just its children
-		for (int32 ChildIndex : Dependencies)
+		FAssetTableRow& Row = OwningTable.GetAssetChecked(CurrentIndex);
+		if (Row.GetGameFeaturePlugin() != ThisGFP)
 		{
-			IndicesToVisit.Add(ChildIndex);
+			// Don't traverse outside this plugin
+			continue;
 		}
+		VisitedIndices.Add(CurrentIndex);
 
-		while (IndicesToVisit.Num() > 0)
+		for (int32 ChildIndex : OwningTable.GetAssetChecked(CurrentIndex).GetDependencies())
 		{
-			int32 CurrentIndex = IndicesToVisit.Pop(false);
-
-			FAssetTableRow* Row = OwningTable->GetAsset(CurrentIndex);
-			if (Row->GameFeaturePlugin != ThisGFP)
+			if (!VisitedIndices.Contains(ChildIndex))
 			{
-				// Don't traverse outside this plugin
-				continue;
-			}
-			VisitedIndices.Add(CurrentIndex);
-
-			for (int32 ChildIndex : OwningTable->GetAsset(CurrentIndex)->GetDependencies())
-			{
-				if (!VisitedIndices.Contains(ChildIndex))
-				{
-					IndicesToVisit.AddUnique(ChildIndex);
-				}
+				IndicesToVisit.AddUnique(ChildIndex);
 			}
 		}
+	}
 
-		// Iteratively separate the graph of "all things referenced by ThisIndex, directly or indirectly"
-		// into "UniqueDependencies -- things referenced ONLY by ThisIndex and by other things themselves referenced ONLY by ThisIndex"
-		// and "OtherDependencies" -- things removed from the list of "all things referenced by ThisIndex" in order to identify UniqueDependencies
-		TSet<int32> UniqueDependencies;
-		TSet<int32> OtherDependencies;
-		while (RefineDependencies(VisitedIndices, OwningTable, ThisIndex, &UniqueDependencies, &OtherDependencies))
-		{
-			VisitedIndices = UniqueDependencies;
-			UniqueDependencies.Empty();
-		}
+	// Iteratively separate the graph of "all things referenced by ThisIndex, directly or indirectly"
+	// into "UniqueDependencies -- things referenced ONLY by ThisIndex and by other things themselves referenced ONLY by ThisIndex"
+	// and "OtherDependencies" -- things removed from the list of "all things referenced by ThisIndex" in order to identify UniqueDependencies
+	TSet<int32> UniqueDependencies;
+	TSet<int32> OtherDependencies;
+	while (RefineDependencies(VisitedIndices, OwningTable, ThisIndex, UniqueDependencies, OtherDependencies))
+	{
+		VisitedIndices = UniqueDependencies;
+		UniqueDependencies.Empty();
+	}
 
-		for (int32 Index : UniqueDependencies)
-		{
-			int64 DependencySize = OwningTable->GetAsset(Index)->GetStagedCompressedSize();
-			TotalSizeUniqueDependencies += DependencySize;
-		}
-		if (OutUniqueDependencies != nullptr)
+	for (int32 Index : UniqueDependencies)
+	{
+		int64 DependencySize = OwningTable.GetAssetChecked(Index).GetStagedCompressedSize();
+		TotalSizeUniqueDependencies += DependencySize;
+	}
+	if (OutUniqueDependencies != nullptr)
+	{
+		*OutUniqueDependencies = UniqueDependencies;
+	}
+
+	// Now explore all the dependencies of OtherDependencies and gather up their sizes
+	// This is necessary because the process calling RefineDependencies doesn't produce a complete list of other dependencies,
+	// just a partial set. By exploring the dependencies of that partial set we can find all the things that were
+	// referenced by ThisIndex but also by some other asset outside the subgraph defined by ThisIndex and its UniqueDependencies
+	VisitedIndices.Empty();
+	IndicesToVisit.Empty();
+	for (int32 Index : OtherDependencies)
+	{
+		IndicesToVisit.Add(Index);
+	}
+
+	while (IndicesToVisit.Num() > 0)
+	{
+		int32 CurrentIndex = IndicesToVisit.Pop(false);
+		FAssetTableRow& Row = OwningTable.GetAssetChecked(CurrentIndex);
+		VisitedIndices.Add(CurrentIndex);
+		if (Row.GetGameFeaturePlugin() != ThisGFP)
 		{
-			*OutUniqueDependencies = UniqueDependencies;
+			// Don't traverse outside this plugin
+			continue;
 		}
 
-		// Now explore all the dependencies of OtherDependencies and gather up their sizes
-		// This is necessary because the process calling RefineDependencies doesn't produce a complete list of other dependencies,
-		// just a partial set. By exploring the dependencies of that partial set we can find all the things that were 
-		// referenced by ThisIndex but also by some other asset outside the subgraph defined by ThisIndex and its UniqueDependencies
-		VisitedIndices.Empty();
-		IndicesToVisit.Empty();
-		for (int32 Index : OtherDependencies)
+		int64 DependencySize = OwningTable.GetAssetChecked(CurrentIndex).GetStagedCompressedSize();
+		TotalSizeOtherDependencies += DependencySize;
+		if (OutSharedDependencies != nullptr)
 		{
-			IndicesToVisit.Add(Index);
+			OutSharedDependencies->Add(CurrentIndex);
 		}
 
-		while (IndicesToVisit.Num() > 0)
+		for (int32 ChildIndex : OwningTable.GetAssetChecked(CurrentIndex).GetDependencies())
 		{
-			int32 CurrentIndex = IndicesToVisit.Pop(false);
-			FAssetTableRow* Row = OwningTable->GetAsset(CurrentIndex);
-			VisitedIndices.Add(CurrentIndex);
-			if (Row->GameFeaturePlugin != ThisGFP)
+			if (!VisitedIndices.Contains(ChildIndex) && ThisIndex != ChildIndex)
 			{
-				// Don't traverse outside this plugin
-				continue;
-			}
-			
-			int64 DependencySize = OwningTable->GetAsset(CurrentIndex)->GetStagedCompressedSize();
-			TotalSizeOtherDependencies += DependencySize;
-			if (OutSharedDependencies != nullptr)
-			{
-				OutSharedDependencies->Add(CurrentIndex);
-			}
-
-			for (int32 ChildIndex : OwningTable->GetAsset(CurrentIndex)->GetDependencies())
-			{
-				if (!VisitedIndices.Contains(ChildIndex) && ThisIndex != ChildIndex)
-				{
-					IndicesToVisit.AddUnique(ChildIndex);
-				}
+				IndicesToVisit.AddUnique(ChildIndex);
 			}
 		}
 	}
 }
 
-int64 FAssetTableRow::GetOrComputeTotalSizeUniqueDependencies(FAssetTable* OwningTable, int32 ThisIndex) const
+////////////////////////////////////////////////////////////////////////////////////////////////////
+
+int64 FAssetTableRow::GetOrComputeTotalSizeUniqueDependencies(FAssetTable& OwningTable, int32 ThisIndex) const
 {
 	if (TotalSizeUniqueDependencies == -1)
 	{
+		UE::Insights::FStopwatch Stopwatch;
+		Stopwatch.Start();
+
 		ComputeDependencySizes(OwningTable, ThisIndex);
+
+		Stopwatch.Stop();
+		const double Duration = Stopwatch.GetAccumulatedTime();
+		if (Duration > 1.0)
+		{
+			UE_LOG(LogInsights, Log, TEXT("ComputeDependencySizes(%s%s) : %.3fs"), Path, Name, Duration);
+		}
 	}
 	return TotalSizeUniqueDependencies;
 }
 
-int64 FAssetTableRow::GetOrComputeTotalSizeOtherDependencies(FAssetTable* OwningTable, int32 ThisIndex) const
+////////////////////////////////////////////////////////////////////////////////////////////////////
+
+int64 FAssetTableRow::GetOrComputeTotalSizeOtherDependencies(FAssetTable& OwningTable, int32 ThisIndex) const
 {
 	if (TotalSizeOtherDependencies == -1)
 	{
+		UE::Insights::FStopwatch Stopwatch;
+		Stopwatch.Start();
+
 		ComputeDependencySizes(OwningTable, ThisIndex);
+
+		Stopwatch.Stop();
+		const double Duration = Stopwatch.GetAccumulatedTime();
+		if (Duration > 1.0)
+		{
+			UE_LOG(LogInsights, Log, TEXT("ComputeDependencySizes(%s%s) : %.3fs"), Path, Name, Duration);
+		}
 	}
 	return TotalSizeOtherDependencies;
 }
diff --git a/Engine/Plugins/Editor/AssetManagerEditor/Source/AssetManagerEditor/Private/TreeView/AssetTable.h b/Engine/Plugins/Editor/AssetManagerEditor/Source/AssetManagerEditor/Private/TreeView/AssetTable.h
--- a/Engine/Plugins/Editor/AssetManagerEditor/Source/AssetManagerEditor/Private/TreeView/AssetTable.h
+++ b/Engine/Plugins/Editor/AssetManagerEditor/Source/AssetManagerEditor/Private/TreeView/AssetTable.h
@@ -138,28 +138,31 @@ public:
 	{
 	}
 
+	FName GetNodeName() const { return FName(Name, 0); }
+
 	const TCHAR* GetType() const { return Type; }
 	const TCHAR* GetName() const { return Name; }
 	const TCHAR* GetPath() const { return Path; }
 	const TCHAR* GetPrimaryType() const { return PrimaryType; }
 	const TCHAR* GetPrimaryName() const { return PrimaryName; }
 	int64 GetStagedCompressedSize() const { return StagedCompressedSize; }
-	int64 GetOrComputeTotalSizeUniqueDependencies(FAssetTable* OwningTable, int32 ThisIndex) const;
-	int64 GetOrComputeTotalSizeOtherDependencies(FAssetTable* OwningTable, int32 ThisIndex) const; 
+	int64 GetOrComputeTotalSizeUniqueDependencies(FAssetTable& OwningTable, int32 ThisIndex) const;
+	int64 GetOrComputeTotalSizeOtherDependencies(FAssetTable& OwningTable, int32 ThisIndex) const; 
 	int64 GetTotalUsageCount() const { return TotalUsageCount; }
 	const TCHAR* GetChunks() const { return Chunks; }
 	const TCHAR* GetNativeClass() const { return NativeClass; }
 	const TCHAR* GetGameFeaturePlugin() const { return GameFeaturePlugin; }
 	int32 GetNumDependencies() const { return Dependencies.Num(); }
 	const TArray<int32>& GetDependencies() const { return Dependencies; }
-	FName GetNodeName() const { return FName(Name, 0); }
+	int32 GetNumReferencers() const { return Referencers.Num(); }
+	const TArray<int32>& GetReferencers() const { return Referencers; }
 	FLinearColor GetColor() const { return Color; }
 
 private:
-	static bool RefineDependencies(TSet<int32> PreviouslyVisitedIndices, FAssetTable* OwningTable, int32 ThisIndex, TSet<int32>* OutIncrementallyRefinedUniqueIndices, TSet<int32>* OutExcludedIndices);
-	// Providing OutUniqueDependencies and OutSharedDependencies will force a full recomputation. It's intended for use in unit testing/validation.
-	void ComputeDependencySizes(FAssetTable* OwningTable, int32 ThisIndex, TSet<int32>* OutUniqueDependencies = nullptr, TSet<int32>* OutSharedDependencies = nullptr) const;
+	static bool RefineDependencies(const TSet<int32>& PreviouslyVisitedIndices, FAssetTable& OwningTable, int32 ThisIndex, TSet<int32>& OutIncrementallyRefinedUniqueIndices, TSet<int32>& OutExcludedIndices);
+	void ComputeDependencySizes(FAssetTable& OwningTable, int32 ThisIndex, TSet<int32>* OutUniqueDependencies = nullptr, TSet<int32>* OutSharedDependencies = nullptr) const;
 
+private:
 	const TCHAR* Type = nullptr;
 	const TCHAR* Name = nullptr;
 	const TCHAR* Path = nullptr;
diff --git a/Engine/Plugins/Editor/AssetManagerEditor/Source/AssetManagerEditor/Private/TreeView/AssetTreeNode.h b/Engine/Plugins/Editor/AssetManagerEditor/Source/AssetManagerEditor/Private/TreeView/AssetTreeNode.h
--- a/Engine/Plugins/Editor/AssetManagerEditor/Source/AssetManagerEditor/Private/TreeView/AssetTreeNode.h
+++ b/Engine/Plugins/Editor/AssetManagerEditor/Source/AssetManagerEditor/Private/TreeView/AssetTreeNode.h
@@ -36,6 +36,7 @@ public:
 
 	virtual ~FAssetTreeNode() {}
 
+	bool IsValidAsset() const { return AssetTable && AssetTable->IsValidRowIndex(RowId.RowIndex); }
 	FAssetTable& GetAssetTableChecked() const { return *AssetTable; }
 	const FAssetTableRow& GetAssetChecked() const { return AssetTable->GetAssetChecked(RowId.RowIndex); }
 
diff --git a/Engine/Plugins/Editor/AssetManagerEditor/Source/AssetManagerEditor/Private/TreeView/SAssetTableTreeView.cpp b/Engine/Plugins/Editor/AssetManagerEditor/Source/AssetManagerEditor/Private/TreeView/SAssetTableTreeView.cpp
--- a/Engine/Plugins/Editor/AssetManagerEditor/Source/AssetManagerEditor/Private/TreeView/SAssetTableTreeView.cpp
+++ b/Engine/Plugins/Editor/AssetManagerEditor/Source/AssetManagerEditor/Private/TreeView/SAssetTableTreeView.cpp
@@ -35,7 +35,6 @@
 
 SAssetTableTreeView::SAssetTableTreeView()
 {
-	bRunInAsyncMode = true;
 }
 
 ////////////////////////////////////////////////////////////////////////////////////////////////////
@@ -48,6 +47,8 @@ SAssetTableTreeView::~SAssetTableTreeView()
 
 void SAssetTableTreeView::Construct(const FArguments& InArgs, TSharedPtr<FAssetTable> InTablePtr)
 {
+	this->OnSelectionChanged = InArgs._OnSelectionChanged;
+
 	STableTreeView::ConstructWidget(InTablePtr);
 
 	CreateGroupings();
@@ -251,18 +252,20 @@ void SAssetTableTreeView::InitAvailableViewPresets()
 		}
 		virtual void GetColumnConfigSet(TArray<UE::Insights::FTableColumnConfig>& InOutConfigSet) const override
 		{
-			InOutConfigSet.Add({ UE::Insights::FTable::GetHierarchyColumnId(),     true, 200.0f });
-			InOutConfigSet.Add({ FAssetTableColumns::CountColumnId,                true, 100.0f });
-			InOutConfigSet.Add({ FAssetTableColumns::NameColumnId,                 true, 200.0f });
-			InOutConfigSet.Add({ FAssetTableColumns::TypeColumnId,                 true, 200.0f });
-			InOutConfigSet.Add({ FAssetTableColumns::PathColumnId,                 true, 400.0f });
-			InOutConfigSet.Add({ FAssetTableColumns::PrimaryTypeColumnId,          true, 200.0f });
-			InOutConfigSet.Add({ FAssetTableColumns::PrimaryNameColumnId,          true, 200.0f });
-			InOutConfigSet.Add({ FAssetTableColumns::StagedCompressedSizeColumnId, true, 100.0f });
-			InOutConfigSet.Add({ FAssetTableColumns::TotalUsageCountColumnId,      true, 100.0f });
-			//InOutConfigSet.Add({ FAssetTableColumns::ChunksColumnId,               true, 200.0f });
-			InOutConfigSet.Add({ FAssetTableColumns::NativeClassColumnId,          true, 200.0f });
-			InOutConfigSet.Add({ FAssetTableColumns::GameFeaturePluginColumnId,    true, 200.0f });
+			InOutConfigSet.Add({ UE::Insights::FTable::GetHierarchyColumnId(),              true, 400.0f });
+			InOutConfigSet.Add({ FAssetTableColumns::CountColumnId,                         true, 100.0f });
+			InOutConfigSet.Add({ FAssetTableColumns::TypeColumnId,                         !true, 200.0f });
+			InOutConfigSet.Add({ FAssetTableColumns::NameColumnId,                         !true, 200.0f });
+			InOutConfigSet.Add({ FAssetTableColumns::PathColumnId,                         !true, 400.0f });
+			InOutConfigSet.Add({ FAssetTableColumns::PrimaryTypeColumnId,                   true, 200.0f });
+			InOutConfigSet.Add({ FAssetTableColumns::PrimaryNameColumnId,                   true, 200.0f });
+			InOutConfigSet.Add({ FAssetTableColumns::StagedCompressedSizeColumnId,          true, 100.0f });
+			InOutConfigSet.Add({ FAssetTableColumns::TotalSizeUniqueDependenciesColumnId,  !true, 100.0f });
+			InOutConfigSet.Add({ FAssetTableColumns::TotalSizeOtherDependenciesColumnId,   !true, 100.0f });
+			InOutConfigSet.Add({ FAssetTableColumns::TotalUsageCountColumnId,               true, 100.0f });
+			InOutConfigSet.Add({ FAssetTableColumns::ChunksColumnId,                       !true, 200.0f });
+			InOutConfigSet.Add({ FAssetTableColumns::NativeClassColumnId,                   true, 200.0f });
+			InOutConfigSet.Add({ FAssetTableColumns::GameFeaturePluginColumnId,             true, 200.0f });
 		}
 	};
 	AvailableViewPresets.Add(MakeShared<FDefaultViewPreset>());
@@ -275,12 +278,12 @@ void SAssetTableTreeView::InitAvailableViewPresets()
 	public:
 		virtual FText GetName() const override
 		{
-			return LOCTEXT("GFPTypeDepView_PresetName", "DependencyAnalysis");
+			return LOCTEXT("GFPTypeDepView_PresetName", "Dependency Analysis");
 		}
 
 		virtual FText GetToolTip() const override
 		{
-			return LOCTEXT("GFPTypeDepView_PresentName", "Dependency Analysis View\nconfigure the tree view to show a breakdown of assets by Game Feature Plugin, Primary Asset Type, and Dependencies.");
+			return LOCTEXT("GFPTypeDepView_PresetToolTip", "Dependency Analysis View\nConfigure the tree view to show a breakdown of assets by Game Feature Plugin, Primary Asset Type, and Dependencies.");
 		}
 		virtual FName GetSortColumn() const override
 		{
@@ -331,17 +334,20 @@ void SAssetTableTreeView::InitAvailableViewPresets()
 		}
 		virtual void GetColumnConfigSet(TArray<UE::Insights::FTableColumnConfig>& InOutConfigSet) const override
 		{
-			InOutConfigSet.Add({ UE::Insights::FTable::GetHierarchyColumnId(),     true, 400.0f });
-			InOutConfigSet.Add({ FAssetTableColumns::CountColumnId,                true, 100.0f });
-			InOutConfigSet.Add({ FAssetTableColumns::NameColumnId,                 true, 200.0f });
-			InOutConfigSet.Add({ FAssetTableColumns::PrimaryTypeColumnId,          true, 200.0f });
-			InOutConfigSet.Add({ FAssetTableColumns::StagedCompressedSizeColumnId, true, 100.0f });
-			InOutConfigSet.Add({ FAssetTableColumns::TotalSizeUniqueDependenciesColumnId, true, 180.0f });
-			InOutConfigSet.Add({ FAssetTableColumns::TotalSizeOtherDependenciesColumnId, true, 180.0f });
-			InOutConfigSet.Add({ FAssetTableColumns::NativeClassColumnId,          true, 100.0f });
-			InOutConfigSet.Add({ FAssetTableColumns::GameFeaturePluginColumnId,    true, 200.0f });
-			InOutConfigSet.Add({ FAssetTableColumns::PrimaryNameColumnId,          true, 200.0f });
-			InOutConfigSet.Add({ FAssetTableColumns::TypeColumnId,                 true, 200.0f });
+			InOutConfigSet.Add({ UE::Insights::FTable::GetHierarchyColumnId(),              true, 400.0f });
+			InOutConfigSet.Add({ FAssetTableColumns::CountColumnId,                         true, 100.0f });
+			InOutConfigSet.Add({ FAssetTableColumns::TypeColumnId,                         !true, 200.0f });
+			InOutConfigSet.Add({ FAssetTableColumns::NameColumnId,                         !true, 200.0f });
+			InOutConfigSet.Add({ FAssetTableColumns::PathColumnId,                         !true, 400.0f });
+			InOutConfigSet.Add({ FAssetTableColumns::PrimaryTypeColumnId,                   true, 200.0f });
+			InOutConfigSet.Add({ FAssetTableColumns::PrimaryNameColumnId,                  !true, 200.0f });
+			InOutConfigSet.Add({ FAssetTableColumns::StagedCompressedSizeColumnId,          true, 100.0f });
+			InOutConfigSet.Add({ FAssetTableColumns::TotalSizeUniqueDependenciesColumnId,   true, 100.0f });
+			InOutConfigSet.Add({ FAssetTableColumns::TotalSizeOtherDependenciesColumnId,    true, 100.0f });
+			InOutConfigSet.Add({ FAssetTableColumns::TotalUsageCountColumnId,              !true, 100.0f });
+			InOutConfigSet.Add({ FAssetTableColumns::ChunksColumnId,                       !true, 200.0f });
+			InOutConfigSet.Add({ FAssetTableColumns::NativeClassColumnId,                   true, 200.0f });
+			InOutConfigSet.Add({ FAssetTableColumns::GameFeaturePluginColumnId,             true, 200.0f });
 		}
 	};
 	AvailableViewPresets.Add(MakeShared<FGameFeaturePluginPrimaryTypeDependencyView>());
@@ -388,21 +394,20 @@ void SAssetTableTreeView::InitAvailableViewPresets()
 		}
 		virtual void GetColumnConfigSet(TArray<UE::Insights::FTableColumnConfig>& InOutConfigSet) const override
 		{
-			InOutConfigSet.Add({ UE::Insights::FTable::GetHierarchyColumnId(),     true, 400.0f });
-			InOutConfigSet.Add({ FAssetTableColumns::CountColumnId,                true, 100.0f });
-			InOutConfigSet.Add({ FAssetTableColumns::NameColumnId,                 true, 200.0f });
-			InOutConfigSet.Add({ FAssetTableColumns::TypeColumnId,                 true, 200.0f });
-			//InOutConfigSet.Add({ FAssetTableColumns::PathColumnId,                 true, 400.0f });
-			InOutConfigSet.Add({ FAssetTableColumns::PrimaryTypeColumnId,          true, 200.0f });
-			InOutConfigSet.Add({ FAssetTableColumns::PrimaryNameColumnId,          true, 200.0f });
-			//InOutConfigSet.Add({ FAssetTableColumns::ManagedDiskSizeColumnId,      true, 100.0f });
-			//InOutConfigSet.Add({ FAssetTableColumns::DiskSizeColumnId,             true, 100.0f });
-			InOutConfigSet.Add({ FAssetTableColumns::StagedCompressedSizeColumnId, true, 100.0f });
-			InOutConfigSet.Add({ FAssetTableColumns::TotalUsageCountColumnId,      true, 100.0f });
-			//InOutConfigSet.Add({ FAssetTableColumns::CookRuleColumnId,             true, 200.0f });
-			//InOutConfigSet.Add({ FAssetTableColumns::ChunksColumnId,               true, 200.0f });
-			InOutConfigSet.Add({ FAssetTableColumns::NativeClassColumnId,          true, 200.0f });
-			InOutConfigSet.Add({ FAssetTableColumns::GameFeaturePluginColumnId,    true, 200.0f });
+			InOutConfigSet.Add({ UE::Insights::FTable::GetHierarchyColumnId(),              true, 400.0f });
+			InOutConfigSet.Add({ FAssetTableColumns::CountColumnId,                         true, 100.0f });
+			InOutConfigSet.Add({ FAssetTableColumns::TypeColumnId,                         !true, 200.0f });
+			InOutConfigSet.Add({ FAssetTableColumns::NameColumnId,                         !true, 200.0f });
+			InOutConfigSet.Add({ FAssetTableColumns::PathColumnId,                         !true, 400.0f });
+			InOutConfigSet.Add({ FAssetTableColumns::PrimaryTypeColumnId,                   true, 200.0f });
+			InOutConfigSet.Add({ FAssetTableColumns::PrimaryNameColumnId,                   true, 200.0f });
+			InOutConfigSet.Add({ FAssetTableColumns::StagedCompressedSizeColumnId,          true, 100.0f });
+			InOutConfigSet.Add({ FAssetTableColumns::TotalSizeUniqueDependenciesColumnId,   true, 100.0f });
+			InOutConfigSet.Add({ FAssetTableColumns::TotalSizeOtherDependenciesColumnId,    true, 100.0f });
+			InOutConfigSet.Add({ FAssetTableColumns::TotalUsageCountColumnId,               true, 100.0f });
+			InOutConfigSet.Add({ FAssetTableColumns::ChunksColumnId,                       !true, 200.0f });
+			InOutConfigSet.Add({ FAssetTableColumns::NativeClassColumnId,                   true, 200.0f });
+			InOutConfigSet.Add({ FAssetTableColumns::GameFeaturePluginColumnId,             true, 200.0f });
 		}
 	};
 	AvailableViewPresets.Add(MakeShared<FAssetPathViewPreset>());
@@ -446,36 +451,23 @@ void SAssetTableTreeView::InitAvailableViewPresets()
 			{
 				InOutCurrentGroupings.Add(*PrimaryTypeGrouping);
 			}
-
-			const TSharedPtr<UE::Insights::FTreeNodeGrouping>* PrimaryNameGrouping = InAvailableGroupings.FindByPredicate(
-				[](TSharedPtr<UE::Insights::FTreeNodeGrouping>& Grouping)
-				{
-					return Grouping->Is<UE::Insights::FTreeNodeGroupingByUniqueValueCString>() &&
-						   Grouping->As<UE::Insights::FTreeNodeGroupingByUniqueValueCString>().GetColumnId() == FAssetTableColumns::PrimaryNameColumnId;
-				});
-			if (PrimaryNameGrouping)
-			{
-				InOutCurrentGroupings.Add(*PrimaryNameGrouping);
-			}
 		}
 		virtual void GetColumnConfigSet(TArray<UE::Insights::FTableColumnConfig>& InOutConfigSet) const override
 		{
-			InOutConfigSet.Add({ UE::Insights::FTable::GetHierarchyColumnId(),     true, 300.0f });
-			InOutConfigSet.Add({ FAssetTableColumns::CountColumnId,                true, 100.0f });
-			InOutConfigSet.Add({ FAssetTableColumns::NameColumnId,                 true, 200.0f });
-			InOutConfigSet.Add({ FAssetTableColumns::TypeColumnId,                 true, 200.0f });
-			InOutConfigSet.Add({ FAssetTableColumns::PathColumnId,                 true, 400.0f });
-			//InOutConfigSet.Add({ FAssetTableColumns::PrimaryTypeColumnId,          true, 200.0f });
-			//InOutConfigSet.Add({ FAssetTableColumns::PrimaryNameColumnId,          true, 200.0f });
-			//InOutConfigSet.Add({ FAssetTableColumns::ManagedDiskSizeColumnId,      true, 100.0f });
-			//InOutConfigSet.Add({ FAssetTableColumns::DiskSizeColumnId,             true, 100.0f });
-			InOutConfigSet.Add({ FAssetTableColumns::StagedCompressedSizeColumnId, true, 100.0f });
-			InOutConfigSet.Add({ FAssetTableColumns::TotalSizeUniqueDependenciesColumnId, true, 100.0f });
-			InOutConfigSet.Add({ FAssetTableColumns::TotalUsageCountColumnId,      true, 100.0f });
-			//InOutConfigSet.Add({ FAssetTableColumns::CookRuleColumnId,             true, 200.0f });
-			//InOutConfigSet.Add({ FAssetTableColumns::ChunksColumnId,               true, 200.0f });
-			InOutConfigSet.Add({ FAssetTableColumns::NativeClassColumnId,          true, 200.0f });
-			InOutConfigSet.Add({ FAssetTableColumns::GameFeaturePluginColumnId,    true, 200.0f });
+			InOutConfigSet.Add({ UE::Insights::FTable::GetHierarchyColumnId(),              true, 400.0f });
+			InOutConfigSet.Add({ FAssetTableColumns::CountColumnId,                         true, 100.0f });
+			InOutConfigSet.Add({ FAssetTableColumns::TypeColumnId,                         !true, 200.0f });
+			InOutConfigSet.Add({ FAssetTableColumns::NameColumnId,                         !true, 200.0f });
+			InOutConfigSet.Add({ FAssetTableColumns::PathColumnId,                         !true, 400.0f });
+			InOutConfigSet.Add({ FAssetTableColumns::PrimaryTypeColumnId,                  !true, 200.0f });
+			InOutConfigSet.Add({ FAssetTableColumns::PrimaryNameColumnId,                  !true, 200.0f });
+			InOutConfigSet.Add({ FAssetTableColumns::StagedCompressedSizeColumnId,          true, 100.0f });
+			InOutConfigSet.Add({ FAssetTableColumns::TotalSizeUniqueDependenciesColumnId,   true, 100.0f });
+			InOutConfigSet.Add({ FAssetTableColumns::TotalSizeOtherDependenciesColumnId,    true, 100.0f });
+			InOutConfigSet.Add({ FAssetTableColumns::TotalUsageCountColumnId,               true, 100.0f });
+			InOutConfigSet.Add({ FAssetTableColumns::ChunksColumnId,                       !true, 200.0f });
+			InOutConfigSet.Add({ FAssetTableColumns::NativeClassColumnId,                   true, 200.0f });
+			InOutConfigSet.Add({ FAssetTableColumns::GameFeaturePluginColumnId,             true, 200.0f });
 		}
 	};
 	AvailableViewPresets.Add(MakeShared<FPrimaryAssetViewPreset>());
@@ -494,19 +486,31 @@ TSharedPtr<SWidget> SAssetTableTreeView::ConstructFooter()
 		SNew(SHorizontalBox)
 
 		+ SHorizontalBox::Slot()
-		.HAlign(HAlign_Left)
+		.AutoWidth()
 		.Padding(2.0f)
 		[
 			SNew(STextBlock)
-			.Text(LOCTEXT("TODO", "<TODO:status>"))
+			.Text(this, &SAssetTableTreeView::GetFooterLeftText)
 		]
 
 		+ SHorizontalBox::Slot()
-		.HAlign(HAlign_Right)
-		.Padding(2.0f)
+		.FillWidth(1.0f)
+
+		+ SHorizontalBox::Slot()
+		.AutoWidth()
+		.Padding(2.0f, 2.0f, 0.0f, 2.0f)
+		[
+			SNew(STextBlock)
+			.Text(this, &SAssetTableTreeView::GetFooterRightText1)
+			.ColorAndOpacity(FSlateColor(EStyleColor::White25))
+		]
+
+		+ SHorizontalBox::Slot()
+		.AutoWidth()
+		.Padding(0.0f, 2.0f, 2.0f, 2.0f)
 		[
 			SNew(STextBlock)
-			.Text(LOCTEXT("TODO", "<TODO:status>"))
+			.Text(this, &SAssetTableTreeView::GetFooterRightText2)
 		];
 }
 END_SLATE_FUNCTION_BUILD_OPTIMIZATION
@@ -560,23 +564,84 @@ void SAssetTableTreeView::AddCustomAdvancedFilters()
 
 ////////////////////////////////////////////////////////////////////////////////////////////////////
 
-TSharedPtr<FAssetTreeNode> SAssetTableTreeView::GetSingleSelectedAssetNode() const
+void SAssetTableTreeView::ExtendMenu(FMenuBuilder& MenuBuilder)
+{
+}
+
+////////////////////////////////////////////////////////////////////////////////////////////////////
+
+FText SAssetTableTreeView::GetFooterLeftText() const
+{
+	return FooterLeftText;
+}
+
+////////////////////////////////////////////////////////////////////////////////////////////////////
+
+FText SAssetTableTreeView::GetFooterRightText1() const
 {
-	if (TreeView->GetNumItemsSelected() == 1)
+	return FooterRightText1;
+}
+
+////////////////////////////////////////////////////////////////////////////////////////////////////
+
+FText SAssetTableTreeView::GetFooterRightText2() const
+{
+	return FooterRightText2;
+}
+
+////////////////////////////////////////////////////////////////////////////////////////////////////
+
+void SAssetTableTreeView::TreeView_OnSelectionChanged(UE::Insights::FTableTreeNodePtr SelectedItem, ESelectInfo::Type SelectInfo)
+{
+	TArray<UE::Insights::FTableTreeNodePtr> SelectedNodes;
+	const int32 NumSelectedNodes = TreeView->GetSelectedItems(SelectedNodes);
+	int32 NumSelectedAssets = 0;
+	FAssetTreeNodePtr NewSelectedAssetNode;
+	for (const UE::Insights::FTableTreeNodePtr& Node : SelectedNodes)
 	{
-		FAssetTreeNodePtr SelectedTreeNode = StaticCastSharedPtr<FAssetTreeNode>(TreeView->GetSelectedItems()[0]);
-		if (SelectedTreeNode.IsValid() && !SelectedTreeNode->IsGroup())
+		if (Node->Is<FAssetTreeNode>() && !Node->IsGroup())
 		{
-			return SelectedTreeNode;
+			NewSelectedAssetNode = StaticCastSharedPtr<FAssetTreeNode>(Node);
+			++NumSelectedAssets;
 		}
 	}
-	return nullptr;
-}
 
-////////////////////////////////////////////////////////////////////////////////////////////////////
+	const int32 VisibleAssetCount = GetAssetTable()->GetVisibleAssetCount();
+	if (NumSelectedAssets == 0)
+	{
+		FooterLeftText = FText::Format(LOCTEXT("FooterLeftTextFmt0", "{0} assets"), FText::AsNumber(VisibleAssetCount));
+		FooterRightText1 = FText();
+		FooterRightText2 = FText();
+	}
+	else if (NumSelectedAssets == 1)
+	{
+		FooterLeftText = FText::Format(LOCTEXT("FooterLeftTextFmt1", "{0} assets (1 selected)"), FText::AsNumber(VisibleAssetCount));
 
-void SAssetTableTreeView::ExtendMenu(FMenuBuilder& MenuBuilder)
-{
+		const FAssetTableRow& AssetTableRow = NewSelectedAssetNode->GetAssetChecked();
+		FooterRightText1 = FText::FromString(AssetTableRow.GetPath());
+		FooterRightText2 = FText::FromString(AssetTableRow.GetName());
+	}
+	else
+	{
+		FooterLeftText = FText::Format(LOCTEXT("FooterLeftTextFmt2", "{0} assets ({1} selected)"), FText::AsNumber(VisibleAssetCount), FText::AsNumber(NumSelectedAssets));
+		FooterRightText1 = FText();
+		FooterRightText2 = FText();
+	}
+
+	if (NumSelectedAssets != 1)
+	{
+		NewSelectedAssetNode.Reset();
+	}
+
+	if (SelectedAssetNode != NewSelectedAssetNode)
+	{
+		SelectedAssetNode = NewSelectedAssetNode;
+	}
+
+	if (OnSelectionChanged.IsBound())
+	{
+		OnSelectionChanged.Execute(SelectedNodes);
+	}
 }
 
 ////////////////////////////////////////////////////////////////////////////////////////////////////
diff --git a/Engine/Plugins/Editor/AssetManagerEditor/Source/AssetManagerEditor/Private/TreeView/SAssetTableTreeView.h b/Engine/Plugins/Editor/AssetManagerEditor/Source/AssetManagerEditor/Private/TreeView/SAssetTableTreeView.h
--- a/Engine/Plugins/Editor/AssetManagerEditor/Source/AssetManagerEditor/Private/TreeView/SAssetTableTreeView.h
+++ b/Engine/Plugins/Editor/AssetManagerEditor/Source/AssetManagerEditor/Private/TreeView/SAssetTableTreeView.h
@@ -3,6 +3,7 @@
 #pragma once
 
 #include "CoreMinimal.h"
+#include "Delegates/DelegateCombinations.h"
 #include "Widgets/DeclarativeSyntaxSupport.h"
 
 #include "TreeView/AssetTable.h"
@@ -12,6 +13,9 @@ class FAssetTreeNode;
 
 class SAssetTableTreeView : public UE::Insights::STableTreeView
 {
+public:
+	DECLARE_DELEGATE_OneParam(FOnSelectionChanged, const TArray<TSharedPtr<UE::Insights::FTableTreeNode>>)
+
 public:
 	/** Default constructor. */
 	SAssetTableTreeView();
@@ -19,7 +23,11 @@ public:
 	/** Virtual destructor. */
 	virtual ~SAssetTableTreeView();
 
-	SLATE_BEGIN_ARGS(SAssetTableTreeView) {}
+	SLATE_BEGIN_ARGS(SAssetTableTreeView)
+		: _OnSelectionChanged()
+	{
+	}
+	SLATE_EVENT(FOnSelectionChanged, OnSelectionChanged)
 	SLATE_END_ARGS()
 
 	/**
@@ -34,8 +42,6 @@ public:
 	TSharedPtr<FAssetTable> GetAssetTable() { return StaticCastSharedPtr<FAssetTable>(GetTable()); }
 	const TSharedPtr<FAssetTable> GetAssetTable() const { return StaticCastSharedPtr<FAssetTable>(GetTable()); }
 
-	//void UpdateSourceTable(TSharedPtr<TraceServices::IAssetTable> SourceTable);
-
 	virtual void Reset();
 
 	/**
@@ -63,13 +69,19 @@ public:
 
 	////////////////////////////////////////////////////////////////////////////////////////////////////
 
+	TSharedPtr<FAssetTreeNode> GetSingleSelectedAssetNode() const { return SelectedAssetNode; }
+
 protected:
 	virtual void InternalCreateGroupings() override;
 
-	TSharedPtr<FAssetTreeNode> GetSingleSelectedAssetNode() const;
-
 	virtual void ExtendMenu(FMenuBuilder& MenuBuilder) override;
 
+	FText GetFooterLeftText() const;
+	FText GetFooterRightText1() const;
+	FText GetFooterRightText2() const;
+
+	virtual void TreeView_OnSelectionChanged(UE::Insights::FTableTreeNodePtr SelectedItem, ESelectInfo::Type SelectInfo) override;
+
 private:
 	bool virtual ApplyCustomAdvancedFilters(const UE::Insights::FTableTreeNodePtr& NodePtr) override;
 	virtual void AddCustomAdvancedFilters() override;
@@ -78,4 +90,12 @@ private:
 
 private:
 	bool bNeedsToRebuild = false;
+
+	FText FooterLeftText;
+	FText FooterRightText1;
+	FText FooterRightText2;
+
+	/** Delegate to invoke when selection changes. */
+	FOnSelectionChanged OnSelectionChanged;
+	TSharedPtr<FAssetTreeNode> SelectedAssetNode;
 };
