Fixed LOCTEXT_NAMESPACE for TableTreeView implementation in AssetManagerEditor plugin to be diffferent than the one in TraceInsights module (fixes text conflict from UI_COMMAND macro).

diff --git a/Engine/Plugins/Editor/AssetManagerEditor/Source/AssetManagerEditor/Private/Insights/Table/Widgets/SAsyncOperationStatus.cpp b/Engine/Plugins/Editor/AssetManagerEditor/Source/AssetManagerEditor/Private/Insights/Table/Widgets/SAsyncOperationStatus.cpp
--- a/Engine/Plugins/Editor/AssetManagerEditor/Source/AssetManagerEditor/Private/Insights/Table/Widgets/SAsyncOperationStatus.cpp
+++ b/Engine/Plugins/Editor/AssetManagerEditor/Source/AssetManagerEditor/Private/Insights/Table/Widgets/SAsyncOperationStatus.cpp
@@ -11,7 +11,7 @@
 #include "Insights/Common/AsyncWorkUtils.h"
 #include "Insights/Common/TimeUtils.h"
 
-#define LOCTEXT_NAMESPACE "SAsyncOperationStatus"
+#define LOCTEXT_NAMESPACE "UE::Insights::SAsyncOperationStatus"
 
 namespace UE
 {
diff --git a/Engine/Plugins/Editor/AssetManagerEditor/Source/AssetManagerEditor/Private/Insights/Table/Widgets/STableTreeView.cpp b/Engine/Plugins/Editor/AssetManagerEditor/Source/AssetManagerEditor/Private/Insights/Table/Widgets/STableTreeView.cpp
--- a/Engine/Plugins/Editor/AssetManagerEditor/Source/AssetManagerEditor/Private/Insights/Table/Widgets/STableTreeView.cpp
+++ b/Engine/Plugins/Editor/AssetManagerEditor/Source/AssetManagerEditor/Private/Insights/Table/Widgets/STableTreeView.cpp
@@ -34,7 +34,7 @@
 
 #include <limits>
 
-#define LOCTEXT_NAMESPACE "STableTreeView"
+#define LOCTEXT_NAMESPACE "UE::Insights::STableTreeView"
 
 namespace UE
 {
diff --git a/Engine/Plugins/Editor/AssetManagerEditor/Source/AssetManagerEditor/Private/Insights/Table/Widgets/STableTreeViewCell.cpp b/Engine/Plugins/Editor/AssetManagerEditor/Source/AssetManagerEditor/Private/Insights/Table/Widgets/STableTreeViewCell.cpp
--- a/Engine/Plugins/Editor/AssetManagerEditor/Source/AssetManagerEditor/Private/Insights/Table/Widgets/STableTreeViewCell.cpp
+++ b/Engine/Plugins/Editor/AssetManagerEditor/Source/AssetManagerEditor/Private/Insights/Table/Widgets/STableTreeViewCell.cpp
@@ -18,7 +18,7 @@
 #include "Insights/Table/ViewModels/TableColumn.h"
 #include "Insights/Table/Widgets/STableTreeViewRow.h"
 
-#define LOCTEXT_NAMESPACE "STableTreeView"
+#define LOCTEXT_NAMESPACE "UE::Insights::STableTreeView"
 
 namespace UE
 {
diff --git a/Engine/Plugins/Editor/AssetManagerEditor/Source/AssetManagerEditor/Private/Insights/Table/Widgets/STableTreeViewRow.cpp b/Engine/Plugins/Editor/AssetManagerEditor/Source/AssetManagerEditor/Private/Insights/Table/Widgets/STableTreeViewRow.cpp
--- a/Engine/Plugins/Editor/AssetManagerEditor/Source/AssetManagerEditor/Private/Insights/Table/Widgets/STableTreeViewRow.cpp
+++ b/Engine/Plugins/Editor/AssetManagerEditor/Source/AssetManagerEditor/Private/Insights/Table/Widgets/STableTreeViewRow.cpp
@@ -14,7 +14,7 @@
 #include "Insights/Table/Widgets/STableTreeViewCell.h"
 #include "Insights/Table/Widgets/STableTreeViewTooltip.h"
 
-#define LOCTEXT_NAMESPACE "STableTreeView"
+#define LOCTEXT_NAMESPACE "UE::Insights::STableTreeView"
 
 namespace UE
 {
diff --git a/Engine/Plugins/Editor/AssetManagerEditor/Source/AssetManagerEditor/Private/Insights/Table/Widgets/STableTreeViewTooltip.cpp b/Engine/Plugins/Editor/AssetManagerEditor/Source/AssetManagerEditor/Private/Insights/Table/Widgets/STableTreeViewTooltip.cpp
--- a/Engine/Plugins/Editor/AssetManagerEditor/Source/AssetManagerEditor/Private/Insights/Table/Widgets/STableTreeViewTooltip.cpp
+++ b/Engine/Plugins/Editor/AssetManagerEditor/Source/AssetManagerEditor/Private/Insights/Table/Widgets/STableTreeViewTooltip.cpp
@@ -16,7 +16,7 @@
 #include "Insights/Table/ViewModels/TableColumn.h"
 #include "Insights/Table/ViewModels/TableTreeNode.h"
 
-#define LOCTEXT_NAMESPACE "STableTreeView"
+#define LOCTEXT_NAMESPACE "UE::Insights::STableTreeView"
 
 namespace UE
 {
