#UE When typing into the address bar of the reference viewer, only show "no assets found" if there are no dependency nodes matching the string that was entered. Previously, it would do this when you typed in non-asset paths, like searchable names

diff --git a/Engine/Plugins/Editor/AssetManagerEditor/Source/AssetManagerEditor/Private/ReferenceViewer/SReferenceViewer.cpp b/Engine/Plugins/Editor/AssetManagerEditor/Source/AssetManagerEditor/Private/ReferenceViewer/SReferenceViewer.cpp
--- a/Engine/Plugins/Editor/AssetManagerEditor/Source/AssetManagerEditor/Private/ReferenceViewer/SReferenceViewer.cpp
+++ b/Engine/Plugins/Editor/AssetManagerEditor/Source/AssetManagerEditor/Private/ReferenceViewer/SReferenceViewer.cpp
@@ -800,30 +800,21 @@ FText SReferenceViewer::GetStatusText() const
 
 void SReferenceViewer::OnAddressBarTextCommitted(const FText& NewText, ETextCommit::Type CommitInfo)
 {
-	TArray<FAssetIdentifier> NewPaths;
 	if (CommitInfo == ETextCommit::OnEnter)
 	{
-		TArray<FAssetData> SelectedAssets;
-
+		TArray<FAssetIdentifier> NewPaths;
 		FAssetIdentifier NewPath = FAssetIdentifier::FromString(NewText.ToString());
 
-		FAssetRegistryModule& AssetRegistryModule = FModuleManager::GetModuleChecked< FAssetRegistryModule >( TEXT("AssetRegistry") );
+		FAssetRegistryModule& AssetRegistryModule = FModuleManager::GetModuleChecked< FAssetRegistryModule >(TEXT("AssetRegistry"));
 		IAssetRegistry* AssetRegistry = &AssetRegistryModule.Get();
 
-		AssetRegistry->GetAssetsByPackageName(FName(*NewText.ToString()), SelectedAssets);
-		if (SelectedAssets.Num() > 0)
+		// Using GetDependencies just to check if NewPath exists in the dependency tree. We do not actually care about the dependencies here
+		TArray<FAssetIdentifier> UnusedDependencies;
+		if (AssetRegistry->GetDependencies(NewPath, UnusedDependencies))
 		{
-			NewPaths.Add(NewPath);
+			NewPaths.Add(FAssetIdentifier::FromString(NewText.ToString()));
 		}
 
-		else if (AssetRegistry->GetAssetsByPath(FName(*NewText.ToString()), SelectedAssets, true))
-		{
-			for (const FAssetData& AssetData : SelectedAssets)
-			{
-				NewPaths.AddUnique(FAssetIdentifier(AssetData.PackageName));
-			}
-		}	
-
 		SetGraphRootIdentifiers(NewPaths);
 	}
 }
