Improvements to asset audit tree view including fixing the type name shown in the Type Column to the asset type display name, fixing a usability issue where adding the type "Object" to the view would be insufficient to include all major assets (some only come up via dependency analysis such as WP cells), generalized asset graph traversal to support exploring multiple assets as a unit, enabled the built in unit tests by default since they're fast, performance improvements for traversal, modified the status bar to show the total self size, unique dependencies size, and shared dependency sizes for a selection.

diff --git a/Engine/Plugins/Editor/AssetManagerEditor/Source/AssetManagerEditor/Private/SAssetAuditBrowser.cpp b/Engine/Plugins/Editor/AssetManagerEditor/Source/AssetManagerEditor/Private/SAssetAuditBrowser.cpp
--- a/Engine/Plugins/Editor/AssetManagerEditor/Source/AssetManagerEditor/Private/SAssetAuditBrowser.cpp
+++ b/Engine/Plugins/Editor/AssetManagerEditor/Source/AssetManagerEditor/Private/SAssetAuditBrowser.cpp
@@ -3,6 +3,7 @@
 #include "SAssetAuditBrowser.h"
 
 #include "AssetRegistry/AssetRegistryModule.h"
+#include "AssetDefinitionRegistry.h"
 #include "Blueprint/BlueprintSupport.h"
 #include "ClassViewerModule.h"
 #include "ContentBrowserDataSource.h"
@@ -737,10 +738,24 @@ void SAssetAuditBrowser::QuickPopulateAssetTableRow(FAssetTableRow& OutRow, cons
 {
 	// Asset Type
 	FString AssetType;
-	if (!EditorModule->GetStringValueForCustomColumn(AssetData, FName(TEXT("Type")), AssetType))
+	if (UAssetDefinitionRegistry* AssetDefinitionRegistry = UAssetDefinitionRegistry::Get())
 	{
-		AssetType = AssetData.AssetClassPath.ToString();
+		if (const UAssetDefinition* AssetDefinition = AssetDefinitionRegistry->GetAssetDefinitionForAsset(AssetData))
+		{
+			if (AssetDefinition)
+			{
+				AssetType = AssetDefinition->GetAssetDisplayName().ToString();
+			}
+		}
 	}
+	if (AssetType.IsEmpty())
+	{
+		if (!EditorModule->GetStringValueForCustomColumn(AssetData, FName(TEXT("Type")), AssetType))
+		{
+			AssetType = AssetData.AssetClassPath.ToString();
+		}
+	}
+	
 	OutRow.Type = AssetTable.StoreStr(AssetType);
 
 	// Asset Name
@@ -849,6 +864,8 @@ void SAssetAuditBrowser::RefreshAssetTableTreeView(const FARFilter& Filter)
 
 	UE_LOG(LogInsights, Log, TEXT("[AssetTree] Build asset table..."));
 
+	// TODO: We probably shouldn't do any of this unless the asset registry has finished loading. Otherwise bad things seem to happen.
+
 	UE::Insights::FStopwatch Stopwatch;
 	Stopwatch.Start();
 
@@ -891,7 +908,6 @@ void SAssetAuditBrowser::RefreshAssetTableTreeView(const FARFilter& Filter)
 			}
 		}
 	}
-	AssetTable->SetVisibleAssetCount(AssetTable->GetAssets().Num());
 
 	int32 EntryIndex = AssetTable->GetAssets().Num();
 
@@ -968,6 +984,10 @@ void SAssetAuditBrowser::RefreshAssetTableTreeView(const FARFilter& Filter)
 			for (const FAssetData& SourceAssetData : AssetsInSourcePackage)
 			{
 				int32* RowIndex = AssetToIndexMap.Find(SourceAssetData);
+				if (RowIndex == nullptr)
+				{
+					UE_LOG(LogInsights, Warning, TEXT("Failed to find asset %s in package %s, source asset index %d. Asset registry loading was %s"), *SourceAssetData.AssetName.ToString(), *SourceAssetData.PackageName.ToString(), SourceAssetIndex, AssetRegistry->IsLoadingAssets() ? TEXT("INCOMPLETE") : TEXT("complete"));
+				}
 				if (ensure(RowIndex != nullptr))
 				{
 					for (int32 Index : IndicesOfDependenciesInRowTableToAddToCurrentSourceAssetRow)
@@ -993,6 +1013,7 @@ void SAssetAuditBrowser::RefreshAssetTableTreeView(const FARFilter& Filter)
 				}
 			}
 		}
+		AssetTable->SetVisibleAssetCount(AssetTable->GetAssets().Num());
 	}
 
 	Stopwatch.Stop();
diff --git a/Engine/Plugins/Editor/AssetManagerEditor/Source/AssetManagerEditor/Private/TreeView/AssetTable.cpp b/Engine/Plugins/Editor/AssetManagerEditor/Source/AssetManagerEditor/Private/TreeView/AssetTable.cpp
--- a/Engine/Plugins/Editor/AssetManagerEditor/Source/AssetManagerEditor/Private/TreeView/AssetTable.cpp
+++ b/Engine/Plugins/Editor/AssetManagerEditor/Source/AssetManagerEditor/Private/TreeView/AssetTable.cpp
@@ -153,7 +153,7 @@ void FAssetTableStringStore::EnumerateStrings(TFunction<void(const FStringView S
 
 FAssetTable::FAssetTable()
 {
-//#define RUN_ASSET_TRAVERSAL_TESTS
+#define RUN_ASSET_TRAVERSAL_TESTS
 #ifdef RUN_ASSET_TRAVERSAL_TESTS
 	for (int32 i = 0; i < 9; i++)
 	{
@@ -365,14 +365,14 @@ FAssetTable::FAssetTable()
 		{
 			TSet<int32> DiscoveredUniqueDependencies;
 			TSet<int32> DiscoveredSharedDependencies;
-			Assets[AssetIndex].ComputeDependencySizes(this, AssetIndex, &DiscoveredUniqueDependencies, &DiscoveredSharedDependencies);
+			Assets[AssetIndex].ComputeDependencySizes(*this, TSet<int32>{AssetIndex}, &DiscoveredUniqueDependencies, &DiscoveredSharedDependencies);
 			ensureAlways(DiscoveredUniqueDependencies.Num() == UniqueDependencies[AssetIndex].Num()
 				&& DiscoveredUniqueDependencies.Intersect(UniqueDependencies[AssetIndex]).Num() == DiscoveredUniqueDependencies.Num());
 			ensureAlways(DiscoveredSharedDependencies.Num() == SharedDependencies[AssetIndex].Num()
 				&& DiscoveredSharedDependencies.Intersect(SharedDependencies[AssetIndex]).Num() == DiscoveredSharedDependencies.Num());
 		}
 	}
-	Assets.Empty();
+	ClearAllData();
 #endif
 #if 0 // debug, mock data
 
@@ -814,8 +814,8 @@ void FAssetTable::AddDefaultColumns()
 
 		Column.SetIndex(ColumnIndex++);
 
-		Column.SetShortName(LOCTEXT("StagedCompressedSizeColumnName", "Staged"));
-		Column.SetTitleName(LOCTEXT("StagedCompressedSizeColumnTitle", "Staged Compressed Size"));
+		Column.SetShortName(LOCTEXT("StagedCompressedSizeColumnName", "Self Size"));
+		Column.SetTitleName(LOCTEXT("StagedCompressedSizeColumnTitle", "Self Size (Compressed)"));
 		Column.SetDescription(LOCTEXT("StagedCompressedSizeColumnDesc", "Compressed size of iostore chunks for this asset's package. Only visible after staging."));
 
 		Column.SetFlags(ETableColumnFlags::ShouldBeVisible | ETableColumnFlags::CanBeHidden | ETableColumnFlags::CanBeFiltered | ETableColumnFlags::IsDynamic);
@@ -1207,32 +1207,84 @@ void FAssetTable::AddDefaultColumns()
 	//////////////////////////////////////////////////
 }
 
+/*static*/TSet<int32> FAssetTableRow::GatherAllReachableNodes(const TArray<int32>& StartingNodes, const FAssetTable& OwningTable, const TSet<int32>& AdditionalNodesToStopAt, const TSet<const TCHAR*>& RestrictToGFPs)
+{
+	// "visit" ThisIndex to seed the exploration
+	TSet<int32> VisitedIndices;
+	TSet<int32> IndicesToVisit;
+	for (int32 StartingNodeIndex : StartingNodes)
+	{
+		VisitedIndices.Add(StartingNodeIndex);
+		for (int32 ChildIndex : OwningTable.GetAssetChecked(StartingNodeIndex).GetDependencies())
+		{
+			IndicesToVisit.Add(ChildIndex);
+		}
+	}
+
+	while (IndicesToVisit.Num() > 0)
+	{
+		TSet<int32>::TIterator Iterator = IndicesToVisit.CreateIterator();
+		int32 CurrentIndex = *Iterator;
+		Iterator.RemoveCurrent();
+
+		const FAssetTableRow& Row = OwningTable.GetAssetChecked(CurrentIndex);
+		if (!RestrictToGFPs.Contains(Row.GetGameFeaturePlugin()) || AdditionalNodesToStopAt.Contains(CurrentIndex))
+		{
+			// Don't traverse outside this plugin
+			continue;
+		}
+
+		VisitedIndices.Add(CurrentIndex);
+
+		for (int32 ChildIndex : OwningTable.GetAssetChecked(CurrentIndex).GetDependencies())
+		{
+			if (!VisitedIndices.Contains(ChildIndex))
+			{
+				IndicesToVisit.Add(ChildIndex);
+			}
+		}
+	}
+
+	return VisitedIndices;
+}
+
 ////////////////////////////////////////////////////////////////////////////////////////////////////
 // Iteratively refines the dependency set into unique and shared sets. PreviouslyVisitedIndices is split in each pass, moving some elements into
 // OutExcludedIndices (preserving its original contents) and putting those that are still potentially uniquely owned by the ThisIndex asset into
 // OutIncrementallyRefinedUniqueIndices. Returns 'true' if another pass is required (i.e., if work was done).
-/*static*/ bool FAssetTableRow::RefineDependencies(const TSet<int32>& PreviouslyVisitedIndices, FAssetTable& OwningTable, int32 ThisIndex, TSet<int32>& OutIncrementallyRefinedUniqueIndices, TSet<int32>& OutExcludedIndices)
+/*static*/ void FAssetTableRow::RefineDependencies(const TSet<int32>& PreviouslyVisitedIndices, const FAssetTable& OwningTable, const TSet<int32> RootIndices, const TSet<const TCHAR*>& RestrictToGFPs, TSet<int32>& OutIncrementallyRefinedUniqueIndices, TSet<int32>& OutExcludedIndices)
 {
-	const TCHAR* ThisGFP = OwningTable.GetAssetChecked(ThisIndex).GetGameFeaturePlugin();
-
 	// "visit" ThisIndex to seed the exploration
-	TArray<int32> IndicesToVisit = OwningTable.GetAssetChecked(ThisIndex).GetDependencies();
+	TSet<int32> IndicesToVisit;
+	for (int32 RootIndex : RootIndices)
+	{
+		IndicesToVisit = IndicesToVisit.Union(TSet<int32>{OwningTable.GetAssetChecked(RootIndex).GetDependencies()});
+	}
+	TSet<int32> NewlyExcludedDependencies;
 
 	while (IndicesToVisit.Num() > 0)
 	{
-		int32 CurrentIndex = IndicesToVisit.Pop(false);
+		TSet<int32>::TIterator Iterator = IndicesToVisit.CreateIterator();
+		int32 CurrentIndex = *Iterator;
+		Iterator.RemoveCurrent();
 
-		FAssetTableRow& Row = OwningTable.GetAssetChecked(CurrentIndex);
-		if (Row.GetGameFeaturePlugin() != ThisGFP)
+		const FAssetTableRow& Row = OwningTable.GetAssetChecked(CurrentIndex);
+		if (!RestrictToGFPs.Contains(Row.GetGameFeaturePlugin()))
 		{
 			// Don't traverse outside this plugin
 			continue;
 		}
 
+		if (OutExcludedIndices.Contains(CurrentIndex))
+		{
+			// We already know not to traverse into this. We'll handle excluding its transitive dependencies later.
+			continue;
+		}
+
 		bool ShouldIncludeInTotal = true;
 		for (int32 ReferencerIndex : Row.Referencers)
 		{
-			if (PreviouslyVisitedIndices.Contains(ReferencerIndex) == false && ThisIndex != ReferencerIndex)
+			if (PreviouslyVisitedIndices.Contains(ReferencerIndex) == false && !RootIndices.Contains(ReferencerIndex))
 			{
 				ShouldIncludeInTotal = false;
 				break;
@@ -1240,7 +1292,7 @@ void FAssetTable::AddDefaultColumns()
 		}
 		if (ShouldIncludeInTotal == false)
 		{
-			OutExcludedIndices.Add(CurrentIndex);
+			NewlyExcludedDependencies.Add(CurrentIndex);
 			continue;
 		}
 		OutIncrementallyRefinedUniqueIndices.Add(CurrentIndex);
@@ -1248,73 +1300,57 @@ void FAssetTable::AddDefaultColumns()
 		for (int32 ChildIndex : OwningTable.GetAssetChecked(CurrentIndex).GetDependencies())
 		{
 			// Don't revisit nodes we've already visited and don't re-add ThisIndex to avoid loops (and to avoid counting ourself)
-			if (!OutIncrementallyRefinedUniqueIndices.Contains(ChildIndex) && ThisIndex != ChildIndex)
+			if (!OutIncrementallyRefinedUniqueIndices.Contains(ChildIndex) && !RootIndices.Contains(ChildIndex))
 			{
-				IndicesToVisit.AddUnique(ChildIndex);
+				IndicesToVisit.Add(ChildIndex);
 			}
 		}
 	}
 
-	return OutIncrementallyRefinedUniqueIndices.Num() != PreviouslyVisitedIndices.Num();
+	// We now have a set of nodes that we know are not uniquely owned
+	// Anything they reference should also be removed from the set of unique dependencies
+	TSet<int32> NewlyExcludedDependenciesAndTheirTransitiveDependencies = GatherAllReachableNodes(NewlyExcludedDependencies.Array(), OwningTable, RootIndices, RestrictToGFPs);
+	OutExcludedIndices = OutExcludedIndices.Union(NewlyExcludedDependenciesAndTheirTransitiveDependencies);
+	OutIncrementallyRefinedUniqueIndices = OutIncrementallyRefinedUniqueIndices.Difference(NewlyExcludedDependenciesAndTheirTransitiveDependencies);
 }
 
 ////////////////////////////////////////////////////////////////////////////////////////////////////
 
-void FAssetTableRow::ComputeDependencySizes(FAssetTable& OwningTable, int32 ThisIndex, TSet<int32>* OutUniqueDependencies, TSet<int32>* OutSharedDependencies) const
+/*static*/ FAssetTableDependencySizes FAssetTableRow::ComputeDependencySizes(const FAssetTable & OwningTable, const TSet<int32>& RootIndices, TSet<int32>* OutUniqueDependencies, TSet<int32>* OutSharedDependencies)
 {
-	TotalSizeUniqueDependencies = 0;
-	TotalSizeOtherDependencies = 0;
-
-	TArray<int32> IndicesToVisit;
-	TSet<int32> VisitedIndices;
-
-	// Break any loops in the dependency graph
-	VisitedIndices.Add(ThisIndex);
+	QUICK_SCOPE_CYCLE_COUNTER(STAT_FAssetTableRow_ComputeDependencySizes);
 
-	const TCHAR* ThisGFP = OwningTable.GetAssetChecked(ThisIndex).GetGameFeaturePlugin();
+	FAssetTableDependencySizes Result;
 
-	// Don't include this asset itself in the total, just its children
-	for (int32 ChildIndex : Dependencies)
+	TSet<const TCHAR*> RestrictToGFPs;
+	for (int32 RootIndex : RootIndices)
 	{
-		IndicesToVisit.Add(ChildIndex);
+		RestrictToGFPs.Add(OwningTable.GetAssetChecked(RootIndex).GetGameFeaturePlugin());
 	}
 
-	while (IndicesToVisit.Num() > 0)
-	{
-		int32 CurrentIndex = IndicesToVisit.Pop(false);
 
-		FAssetTableRow& Row = OwningTable.GetAssetChecked(CurrentIndex);
-		if (Row.GetGameFeaturePlugin() != ThisGFP)
-		{
-			// Don't traverse outside this plugin
-			continue;
-		}
-		VisitedIndices.Add(CurrentIndex);
-
-		for (int32 ChildIndex : OwningTable.GetAssetChecked(CurrentIndex).GetDependencies())
-		{
-			if (!VisitedIndices.Contains(ChildIndex))
-			{
-				IndicesToVisit.AddUnique(ChildIndex);
-			}
-		}
-	}
+	TSet<int32> VisitedIndices = GatherAllReachableNodes(RootIndices.Array(), OwningTable, TSet<int32>(), RestrictToGFPs);
 
 	// Iteratively separate the graph of "all things referenced by ThisIndex, directly or indirectly"
 	// into "UniqueDependencies -- things referenced ONLY by ThisIndex and by other things themselves referenced ONLY by ThisIndex"
 	// and "OtherDependencies" -- things removed from the list of "all things referenced by ThisIndex" in order to identify UniqueDependencies
 	TSet<int32> UniqueDependencies;
 	TSet<int32> OtherDependencies;
-	while (RefineDependencies(VisitedIndices, OwningTable, ThisIndex, UniqueDependencies, OtherDependencies))
-	{
-		VisitedIndices = UniqueDependencies;
-		UniqueDependencies.Empty();
-	}
+	int32 IterationCount = 0;
+	RefineDependencies(VisitedIndices, OwningTable, RootIndices, RestrictToGFPs, UniqueDependencies, OtherDependencies);
+
+	// If there's only one root provided, the dependencies shouldn't include it. If more than one root is provided,
+	// some roots might be included as dependencies of other roots
+	ensure(RootIndices.Num() > 1 || UniqueDependencies.Intersect(RootIndices).IsEmpty());
+
+	// if there is more than one root index, remove any root indices from the dependency lists in order to ensure they aren't double counted
+	// (they should be accounted for as part of the self sizes of the selected assets)
+	UniqueDependencies = UniqueDependencies.Difference(RootIndices);
+	OtherDependencies = OtherDependencies.Difference(RootIndices);
 
 	for (int32 Index : UniqueDependencies)
 	{
-		int64 DependencySize = OwningTable.GetAssetChecked(Index).GetStagedCompressedSize();
-		TotalSizeUniqueDependencies += DependencySize;
+		Result.UniqueDependenciesSize += OwningTable.GetAssetChecked(Index).GetStagedCompressedSize();
 	}
 	if (OutUniqueDependencies != nullptr)
 	{
@@ -1326,25 +1362,23 @@ void FAssetTableRow::ComputeDependencySizes(FAssetTable& OwningTable, int32 This
 	// just a partial set. By exploring the dependencies of that partial set we can find all the things that were
 	// referenced by ThisIndex but also by some other asset outside the subgraph defined by ThisIndex and its UniqueDependencies
 	VisitedIndices.Empty();
-	IndicesToVisit.Empty();
-	for (int32 Index : OtherDependencies)
-	{
-		IndicesToVisit.Add(Index);
-	}
 
-	while (IndicesToVisit.Num() > 0)
+	while (OtherDependencies.Num() > 0)
 	{
-		int32 CurrentIndex = IndicesToVisit.Pop(false);
-		FAssetTableRow& Row = OwningTable.GetAssetChecked(CurrentIndex);
+		TSet<int32>::TIterator Iterator = OtherDependencies.CreateIterator();
+		int32 CurrentIndex = *Iterator;
+		Iterator.RemoveCurrent();
+
+		const FAssetTableRow& Row = OwningTable.GetAssetChecked(CurrentIndex);
 		VisitedIndices.Add(CurrentIndex);
-		if (Row.GetGameFeaturePlugin() != ThisGFP)
+		if (!RestrictToGFPs.Contains(Row.GetGameFeaturePlugin()))
 		{
 			// Don't traverse outside this plugin
 			continue;
 		}
 
 		int64 DependencySize = OwningTable.GetAssetChecked(CurrentIndex).GetStagedCompressedSize();
-		TotalSizeOtherDependencies += DependencySize;
+		Result.OtherDependenciesSize += DependencySize;
 		if (OutSharedDependencies != nullptr)
 		{
 			OutSharedDependencies->Add(CurrentIndex);
@@ -1352,28 +1386,32 @@ void FAssetTableRow::ComputeDependencySizes(FAssetTable& OwningTable, int32 This
 
 		for (int32 ChildIndex : OwningTable.GetAssetChecked(CurrentIndex).GetDependencies())
 		{
-			if (!VisitedIndices.Contains(ChildIndex) && ThisIndex != ChildIndex)
+			if (!VisitedIndices.Contains(ChildIndex) && !RootIndices.Contains(ChildIndex))
 			{
-				IndicesToVisit.AddUnique(ChildIndex);
+				OtherDependencies.Add(ChildIndex);
 			}
 		}
 	}
+
+	return Result;
 }
 
 ////////////////////////////////////////////////////////////////////////////////////////////////////
 
-int64 FAssetTableRow::GetOrComputeTotalSizeUniqueDependencies(FAssetTable& OwningTable, int32 ThisIndex) const
+int64 FAssetTableRow::GetOrComputeTotalSizeUniqueDependencies(const FAssetTable& OwningTable, int32 ThisIndex) const
 {
 	if (TotalSizeUniqueDependencies == -1)
 	{
 		UE::Insights::FStopwatch Stopwatch;
 		Stopwatch.Start();
 
-		ComputeDependencySizes(OwningTable, ThisIndex);
+		FAssetTableDependencySizes Sizes = ComputeDependencySizes(OwningTable, TSet<int32>{ThisIndex});
+		TotalSizeUniqueDependencies = Sizes.UniqueDependenciesSize;
+		TotalSizeOtherDependencies = Sizes.OtherDependenciesSize;
 
 		Stopwatch.Stop();
 		const double Duration = Stopwatch.GetAccumulatedTime();
-		if (Duration > 1.0)
+		if (Duration > 0.2f)
 		{
 			UE_LOG(LogInsights, Log, TEXT("ComputeDependencySizes(%s%s) : %.3fs"), Path, Name, Duration);
 		}
@@ -1383,14 +1421,16 @@ int64 FAssetTableRow::GetOrComputeTotalSizeUniqueDependencies(FAssetTable& Ownin
 
 ////////////////////////////////////////////////////////////////////////////////////////////////////
 
-int64 FAssetTableRow::GetOrComputeTotalSizeOtherDependencies(FAssetTable& OwningTable, int32 ThisIndex) const
+int64 FAssetTableRow::GetOrComputeTotalSizeOtherDependencies(const FAssetTable& OwningTable, int32 ThisIndex) const
 {
 	if (TotalSizeOtherDependencies == -1)
 	{
 		UE::Insights::FStopwatch Stopwatch;
 		Stopwatch.Start();
 
-		ComputeDependencySizes(OwningTable, ThisIndex);
+		FAssetTableDependencySizes Sizes = ComputeDependencySizes(OwningTable, TSet<int32>{ThisIndex});
+		TotalSizeUniqueDependencies = Sizes.UniqueDependenciesSize;
+		TotalSizeOtherDependencies = Sizes.OtherDependenciesSize;
 
 		Stopwatch.Stop();
 		const double Duration = Stopwatch.GetAccumulatedTime();
diff --git a/Engine/Plugins/Editor/AssetManagerEditor/Source/AssetManagerEditor/Private/TreeView/AssetTable.h b/Engine/Plugins/Editor/AssetManagerEditor/Source/AssetManagerEditor/Private/TreeView/AssetTable.h
--- a/Engine/Plugins/Editor/AssetManagerEditor/Source/AssetManagerEditor/Private/TreeView/AssetTable.h
+++ b/Engine/Plugins/Editor/AssetManagerEditor/Source/AssetManagerEditor/Private/TreeView/AssetTable.h
@@ -125,6 +125,12 @@ struct FAssetTableColumns
 	static const FName GameFeaturePluginColumnId;
 };
 
+struct FAssetTableDependencySizes
+{
+	int64 UniqueDependenciesSize = 0;
+	int64 OtherDependenciesSize = 0;
+};
+
 class FAssetTableRow
 {
 	friend class FAssetTable;
@@ -146,8 +152,8 @@ public:
 	const TCHAR* GetPrimaryType() const { return PrimaryType; }
 	const TCHAR* GetPrimaryName() const { return PrimaryName; }
 	int64 GetStagedCompressedSize() const { return StagedCompressedSize; }
-	int64 GetOrComputeTotalSizeUniqueDependencies(FAssetTable& OwningTable, int32 ThisIndex) const;
-	int64 GetOrComputeTotalSizeOtherDependencies(FAssetTable& OwningTable, int32 ThisIndex) const; 
+	int64 GetOrComputeTotalSizeUniqueDependencies(const FAssetTable& OwningTable, int32 ThisIndex) const;
+	int64 GetOrComputeTotalSizeOtherDependencies(const FAssetTable& OwningTable, int32 ThisIndex) const; 
 	int64 GetTotalUsageCount() const { return TotalUsageCount; }
 	const TCHAR* GetChunks() const { return Chunks; }
 	const TCHAR* GetNativeClass() const { return NativeClass; }
@@ -158,9 +164,13 @@ public:
 	const TArray<int32>& GetReferencers() const { return Referencers; }
 	FLinearColor GetColor() const { return Color; }
 
+	static FAssetTableDependencySizes ComputeDependencySizes(const FAssetTable& OwningTable, const TSet<int32>& RootIndices, TSet<int32>* OutUniqueDependencies = nullptr, TSet<int32>* OutSharedDependencies = nullptr);
+
 private:
-	static bool RefineDependencies(const TSet<int32>& PreviouslyVisitedIndices, FAssetTable& OwningTable, int32 ThisIndex, TSet<int32>& OutIncrementallyRefinedUniqueIndices, TSet<int32>& OutExcludedIndices);
-	void ComputeDependencySizes(FAssetTable& OwningTable, int32 ThisIndex, TSet<int32>* OutUniqueDependencies = nullptr, TSet<int32>* OutSharedDependencies = nullptr) const;
+	static void RefineDependencies(const TSet<int32>& PreviouslyVisitedIndices, const FAssetTable& OwningTable, const TSet<int32> RootIndices, const TSet<const TCHAR*>& RestrictToGFPs, TSet<int32>& OutIncrementallyRefinedUniqueIndices, TSet<int32>& OutExcludedIndices);
+
+	// Finds all nodes reachable from StartingNodes INCLUDING those in StartingNodes
+	static TSet<int32> GatherAllReachableNodes(const TArray<int32>& StartingNodes, const FAssetTable& OwningTable, const TSet<int32>& AdditionalNodesToStopAt, const TSet<const TCHAR*>& RestrictToGFPs);
 
 private:
 	const TCHAR* Type = nullptr;
diff --git a/Engine/Plugins/Editor/AssetManagerEditor/Source/AssetManagerEditor/Private/TreeView/SAssetTableTreeView.cpp b/Engine/Plugins/Editor/AssetManagerEditor/Source/AssetManagerEditor/Private/TreeView/SAssetTableTreeView.cpp
--- a/Engine/Plugins/Editor/AssetManagerEditor/Source/AssetManagerEditor/Private/TreeView/SAssetTableTreeView.cpp
+++ b/Engine/Plugins/Editor/AssetManagerEditor/Source/AssetManagerEditor/Private/TreeView/SAssetTableTreeView.cpp
@@ -502,7 +502,7 @@ TSharedPtr<SWidget> SAssetTableTreeView::ConstructFooter()
 		.Padding(2.0f, 2.0f, 0.0f, 2.0f)
 		[
 			SNew(STextBlock)
-			.Text(this, &SAssetTableTreeView::GetFooterRightText1)
+			.Text(this, &SAssetTableTreeView::GetFooterCenterText1)
 			.ColorAndOpacity(FSlateColor(EStyleColor::White25))
 		]
 
@@ -511,7 +511,18 @@ TSharedPtr<SWidget> SAssetTableTreeView::ConstructFooter()
 		.Padding(0.0f, 2.0f, 2.0f, 2.0f)
 		[
 			SNew(STextBlock)
-			.Text(this, &SAssetTableTreeView::GetFooterRightText2)
+			.Text(this, &SAssetTableTreeView::GetFooterCenterText2)
+		]
+
+		+ SHorizontalBox::Slot()
+		.FillWidth(1.0f)
+
+		+SHorizontalBox::Slot()
+		.AutoWidth()
+		.Padding(2.0f, 2.0f, 2.0f, 2.0f)
+		[
+			SNew(STextBlock)
+			.Text(this, &SAssetTableTreeView::GetFooterRightText1)
 		];
 }
 END_SLATE_FUNCTION_BUILD_OPTIMIZATION
@@ -578,16 +589,23 @@ FText SAssetTableTreeView::GetFooterLeftText() const
 
 ////////////////////////////////////////////////////////////////////////////////////////////////////
 
-FText SAssetTableTreeView::GetFooterRightText1() const
+FText SAssetTableTreeView::GetFooterCenterText1() const
 {
-	return FooterRightText1;
+	return FooterCenterText1;
 }
 
 ////////////////////////////////////////////////////////////////////////////////////////////////////
 
-FText SAssetTableTreeView::GetFooterRightText2() const
+FText SAssetTableTreeView::GetFooterCenterText2() const
 {
-	return FooterRightText2;
+	return FooterCenterText2;
+}
+
+////////////////////////////////////////////////////////////////////////////////////////////////////
+
+FText SAssetTableTreeView::GetFooterRightText1() const
+{
+	return FooterRightText1;
 }
 
 ////////////////////////////////////////////////////////////////////////////////////////////////////
@@ -598,11 +616,14 @@ void SAssetTableTreeView::TreeView_OnSelectionChanged(UE::Insights::FTableTreeNo
 	const int32 NumSelectedNodes = TreeView->GetSelectedItems(SelectedNodes);
 	int32 NumSelectedAssets = 0;
 	FAssetTreeNodePtr NewSelectedAssetNode;
+	TSet<int32> SelectionSetIndices;
+
 	for (const UE::Insights::FTableTreeNodePtr& Node : SelectedNodes)
 	{
 		if (Node->Is<FAssetTreeNode>() && !Node->IsGroup())
 		{
 			NewSelectedAssetNode = StaticCastSharedPtr<FAssetTreeNode>(Node);
+			SelectionSetIndices.Add(NewSelectedAssetNode->GetRowId().RowIndex);
 			++NumSelectedAssets;
 		}
 	}
@@ -611,22 +632,38 @@ void SAssetTableTreeView::TreeView_OnSelectionChanged(UE::Insights::FTableTreeNo
 	if (NumSelectedAssets == 0)
 	{
 		FooterLeftText = FText::Format(LOCTEXT("FooterLeftTextFmt0", "{0} assets"), FText::AsNumber(VisibleAssetCount));
+		FooterCenterText1 = FText();
+		FooterCenterText2 = FText();
 		FooterRightText1 = FText();
-		FooterRightText2 = FText();
 	}
 	else if (NumSelectedAssets == 1)
 	{
 		FooterLeftText = FText::Format(LOCTEXT("FooterLeftTextFmt1", "{0} assets (1 selected)"), FText::AsNumber(VisibleAssetCount));
 
 		const FAssetTableRow& AssetTableRow = NewSelectedAssetNode->GetAssetChecked();
-		FooterRightText1 = FText::FromString(AssetTableRow.GetPath());
-		FooterRightText2 = FText::FromString(AssetTableRow.GetName());
+		FooterCenterText1 = FText::FromString(AssetTableRow.GetPath());
+		FooterCenterText2 = FText::FromString(AssetTableRow.GetName());
+		FooterRightText1 = FText::Format(LOCTEXT("FooterRightFmt3", "Self: {0} Unique: {1} Shared: {2}"),
+			FText::AsMemory(AssetTableRow.GetStagedCompressedSize()),
+			FText::AsMemory(AssetTableRow.GetOrComputeTotalSizeUniqueDependencies(*GetAssetTable(), NewSelectedAssetNode->GetRowId().RowIndex)),
+			FText::AsMemory(AssetTableRow.GetOrComputeTotalSizeOtherDependencies(*GetAssetTable(), NewSelectedAssetNode->GetRowId().RowIndex)));
 	}
 	else
 	{
 		FooterLeftText = FText::Format(LOCTEXT("FooterLeftTextFmt2", "{0} assets ({1} selected)"), FText::AsNumber(VisibleAssetCount), FText::AsNumber(NumSelectedAssets));
-		FooterRightText1 = FText();
-		FooterRightText2 = FText();
+		FooterCenterText1 = FText();
+		FooterCenterText2 = FText();
+
+		int64 TotalSelfSize = 0;
+		for (int32 Index : SelectionSetIndices)
+		{
+			TotalSelfSize += GetAssetTable()->GetAssetChecked(Index).GetStagedCompressedSize();
+		}
+		FAssetTableDependencySizes Sizes = FAssetTableRow::ComputeDependencySizes(*GetAssetTable(), SelectionSetIndices, nullptr, nullptr);
+		FooterRightText1 = FText::Format(LOCTEXT("FooterRightFmt3", "Self: {0} Unique: {1} Shared: {2}"),
+			FText::AsMemory(TotalSelfSize),
+			FText::AsMemory(Sizes.UniqueDependenciesSize),
+			FText::AsMemory(Sizes.OtherDependenciesSize));
 	}
 
 	if (NumSelectedAssets != 1)
diff --git a/Engine/Plugins/Editor/AssetManagerEditor/Source/AssetManagerEditor/Private/TreeView/SAssetTableTreeView.h b/Engine/Plugins/Editor/AssetManagerEditor/Source/AssetManagerEditor/Private/TreeView/SAssetTableTreeView.h
--- a/Engine/Plugins/Editor/AssetManagerEditor/Source/AssetManagerEditor/Private/TreeView/SAssetTableTreeView.h
+++ b/Engine/Plugins/Editor/AssetManagerEditor/Source/AssetManagerEditor/Private/TreeView/SAssetTableTreeView.h
@@ -77,8 +77,9 @@ protected:
 	virtual void ExtendMenu(FMenuBuilder& MenuBuilder) override;
 
 	FText GetFooterLeftText() const;
+	FText GetFooterCenterText1() const;
+	FText GetFooterCenterText2() const;
 	FText GetFooterRightText1() const;
-	FText GetFooterRightText2() const;
 
 	virtual void TreeView_OnSelectionChanged(UE::Insights::FTableTreeNodePtr SelectedItem, ESelectInfo::Type SelectInfo) override;
 
@@ -92,8 +93,9 @@ private:
 	bool bNeedsToRebuild = false;
 
 	FText FooterLeftText;
+	FText FooterCenterText1;
+	FText FooterCenterText2;
 	FText FooterRightText1;
-	FText FooterRightText2;
 
 	/** Delegate to invoke when selection changes. */
 	FOnSelectionChanged OnSelectionChanged;
