Various improvements to the asset audit tooling including: renamed Other to Shared, added support for bubbling up unique/shared dep sizes, added an "External" dependency category, added a context menu option to export dependency data to a CSV

diff --git a/Engine/Plugins/Editor/AssetManagerEditor/Source/AssetManagerEditor/Private/Insights/Table/ViewModels/TableTreeNode.h b/Engine/Plugins/Editor/AssetManagerEditor/Source/AssetManagerEditor/Private/Insights/Table/ViewModels/TableTreeNode.h
--- a/Engine/Plugins/Editor/AssetManagerEditor/Source/AssetManagerEditor/Private/Insights/Table/ViewModels/TableTreeNode.h
+++ b/Engine/Plugins/Editor/AssetManagerEditor/Source/AssetManagerEditor/Private/Insights/Table/ViewModels/TableTreeNode.h
@@ -71,6 +71,15 @@ public:
 	{
 	}
 
+	/** Initialization constructor for a table record node. */
+	explicit FTableTreeNode(const FName InName, TWeakPtr<FTable> InParentTable, int32 InRowIndex, bool IsGroup)
+		: FBaseTreeNode(InName, IsGroup)
+		, ParentTable(InParentTable)
+		, RowId(InRowIndex)
+		, AggregatedValues(nullptr)
+	{
+	}
+
 	virtual ~FTableTreeNode()
 	{
 		CleanupAggregatedValues();
@@ -128,16 +137,10 @@ public:
 		return AggregatedValues->FindChecked(ColumnId);
 	}
 
-	void AddAggregatedValue(const FName& ColumnId, const FTableCellValue& Value)
-	{
-		InitAggregatedValues();
-		AggregatedValues->Add(ColumnId, Value);
-	}
-
 	void SetAggregatedValue(const FName& ColumnId, const FTableCellValue& Value)
 	{
 		InitAggregatedValues();
-		(*AggregatedValues)[ColumnId] = Value;
+		AggregatedValues->Add(ColumnId, Value);
 	}
 
 	//////////////////////////////////////////////////
@@ -165,8 +168,8 @@ class FCustomTableTreeNode : public FTableTreeNode
 
 public:
 	/** Initialization constructor for a table record node. */
-	explicit FCustomTableTreeNode(const FName InName, TWeakPtr<FTable> InParentTable, int32 InRowIndex, const FSlateBrush* InIconBrush, FLinearColor InColor)
-		: FTableTreeNode(InName, InParentTable)
+	explicit FCustomTableTreeNode(const FName InName, TWeakPtr<FTable> InParentTable, int32 InRowIndex, const FSlateBrush* InIconBrush, FLinearColor InColor, bool IsGroup)
+		: FTableTreeNode(InName, InParentTable, InRowIndex, IsGroup)
 		, IconBrush(InIconBrush)
 		, Color(InColor)
 	{
diff --git a/Engine/Plugins/Editor/AssetManagerEditor/Source/AssetManagerEditor/Private/Insights/Table/Widgets/STableTreeView.cpp b/Engine/Plugins/Editor/AssetManagerEditor/Source/AssetManagerEditor/Private/Insights/Table/Widgets/STableTreeView.cpp
--- a/Engine/Plugins/Editor/AssetManagerEditor/Source/AssetManagerEditor/Private/Insights/Table/Widgets/STableTreeView.cpp
+++ b/Engine/Plugins/Editor/AssetManagerEditor/Source/AssetManagerEditor/Private/Insights/Table/Widgets/STableTreeView.cpp
@@ -1560,7 +1560,7 @@ void STableTreeView::UpdateCStringSameValueAggregationSingleNode(const FTableCol
 
 	if (AggregatedValue != nullptr)
 	{
-		InOutGroupNode.AddAggregatedValue(InColumn.GetId(), FTableCellValue(AggregatedValue));
+		InOutGroupNode.SetAggregatedValue(InColumn.GetId(), FTableCellValue(AggregatedValue));
 	}
 	else
 	{
@@ -1621,7 +1621,7 @@ void STableTreeView::UpdateAggregation(const FTableColumn& InColumn, FTableTreeN
 
 	if (bSetInitialValue || AggregatedValue != InitialAggregatedValue)
 	{
-		InOutGroupNode.AddAggregatedValue(InColumn.GetId(), FTableCellValue(AggregatedValue));
+		InOutGroupNode.SetAggregatedValue(InColumn.GetId(), FTableCellValue(AggregatedValue));
 	}
 }
 
diff --git a/Engine/Plugins/Editor/AssetManagerEditor/Source/AssetManagerEditor/Private/TreeView/AssetDependencyGrouping.cpp b/Engine/Plugins/Editor/AssetManagerEditor/Source/AssetManagerEditor/Private/TreeView/AssetDependencyGrouping.cpp
--- a/Engine/Plugins/Editor/AssetManagerEditor/Source/AssetManagerEditor/Private/TreeView/AssetDependencyGrouping.cpp
+++ b/Engine/Plugins/Editor/AssetManagerEditor/Source/AssetManagerEditor/Private/TreeView/AssetDependencyGrouping.cpp
@@ -2,6 +2,7 @@
 
 #include "AssetDependencyGrouping.h"
 #include "AssetTreeNode.h"
+#include "AssetTable.h"
 #include "Insights/Common/AsyncOperationProgress.h"
 
 #define LOCTEXT_NAMESPACE "FAssetDependencyGrouping"
@@ -90,13 +91,13 @@ void FAssetDependencyGrouping::GroupNodes(const TArray<UE::Insights::FTableTreeN
 		if (Asset.GetNumDependencies() > 0)
 		{
 			// Create a group for the asset node (self) + dependencies.
-			FTableTreeNodePtr AssetGroupPtr = MakeShared<FCustomTableTreeNode>(Asset.GetNodeName(), InParentTable, AssetGroupNodeIconBrush, Asset.GetColor());
+			FTableTreeNodePtr AssetGroupPtr = MakeShared<FCustomTableTreeNode>(Asset.GetNodeName(), InParentTable, AssetNode.GetRowIndex(), AssetGroupNodeIconBrush, Asset.GetColor(), true);
 			AssetGroupPtr->SetExpansion(false);
 			ParentGroup.AddChildAndSetGroupPtr(AssetGroupPtr);
 
 			// Add the asset node (self) under a "_self" group node.
 			static FName SelfGroupName(TEXT("_Self_")); // used _ prefix to sort before "Dependencies"
-			FTableTreeNodePtr SelfGroupPtr = MakeShared<FCustomTableTreeNode>(SelfGroupName, InParentTable, AssetGroupNodeIconBrush, Asset.GetColor());
+			FTableTreeNodePtr SelfGroupPtr = MakeShared<FCustomTableTreeNode>(SelfGroupName, InParentTable, AssetNode.GetRowIndex(), AssetGroupNodeIconBrush, Asset.GetColor(), true);
 			SelfGroupPtr->SetExpansion(false);
 			SelfGroupPtr->AddChildAndSetGroupPtr(NodePtr);
 			AssetGroupPtr->AddChildAndSetGroupPtr(SelfGroupPtr);
diff --git a/Engine/Plugins/Editor/AssetManagerEditor/Source/AssetManagerEditor/Private/TreeView/AssetTable.cpp b/Engine/Plugins/Editor/AssetManagerEditor/Source/AssetManagerEditor/Private/TreeView/AssetTable.cpp
--- a/Engine/Plugins/Editor/AssetManagerEditor/Source/AssetManagerEditor/Private/TreeView/AssetTable.cpp
+++ b/Engine/Plugins/Editor/AssetManagerEditor/Source/AssetManagerEditor/Private/TreeView/AssetTable.cpp
@@ -29,7 +29,8 @@ const FName FAssetTableColumns::PrimaryTypeColumnId(TEXT("PrimaryType"));
 const FName FAssetTableColumns::PrimaryNameColumnId(TEXT("PrimaryName"));
 const FName FAssetTableColumns::StagedCompressedSizeColumnId(TEXT("StagedCompressedSize"));
 const FName FAssetTableColumns::TotalSizeUniqueDependenciesColumnId(TEXT("TotalSizeUniqueDependencies"));
-const FName FAssetTableColumns::TotalSizeOtherDependenciesColumnId(TEXT("TotalSizeOtherDependencies"));
+const FName FAssetTableColumns::TotalSizeSharedDependenciesColumnId(TEXT("TotalSizeSharedDependencies"));
+const FName FAssetTableColumns::TotalSizeExternalDependenciesColumnId(TEXT("TotalSizeExternalDependencies"));
 const FName FAssetTableColumns::TotalUsageCountColumnId(TEXT("TotalUsageCount"));
 const FName FAssetTableColumns::ChunksColumnId(TEXT("Chunks"));
 const FName FAssetTableColumns::NativeClassColumnId(TEXT("NativeClass"));
@@ -892,7 +893,14 @@ void FAssetTable::AddDefaultColumns()
 				if (Node.IsGroup())
 				{
 					const FTableTreeNode& NodePtr = static_cast<const FTableTreeNode&>(Node);
-					if (NodePtr.HasAggregatedValue(Column.GetId()))
+					if (NodePtr.GetRowId().HasValidIndex())
+					{
+						// This group node also has a row index, so it conceptually only represents a single row (i.e., it's a synthetic group created by the asset dependency grouping)
+						TSharedPtr<FTable> TablePtr = NodePtr.GetParentTable().Pin();
+						const FAssetTable& AssetTable = static_cast<const FAssetTable&>(*TablePtr);
+						return FTableCellValue(static_cast<int64>(AssetTable.GetAssetChecked(NodePtr.GetRowIndex()).GetOrComputeTotalSizeUniqueDependencies(AssetTable, NodePtr.GetRowIndex())));
+					}
+					else if (NodePtr.HasAggregatedValue(Column.GetId()))
 					{
 						return NodePtr.GetAggregatedValue(Column.GetId());
 					}
@@ -925,16 +933,16 @@ void FAssetTable::AddDefaultColumns()
 		AddColumn(ColumnRef);
 	}
 	//////////////////////////////////////////////////
-	// Total Size of Other Dependencies
+	// Total Size of Shared Dependencies
 	{
-		TSharedRef<FTableColumn> ColumnRef = MakeShared<FTableColumn>(FAssetTableColumns::TotalSizeOtherDependenciesColumnId);
+		TSharedRef<FTableColumn> ColumnRef = MakeShared<FTableColumn>(FAssetTableColumns::TotalSizeSharedDependenciesColumnId);
 		FTableColumn& Column = *ColumnRef;
 
 		Column.SetIndex(ColumnIndex++);
 
-		Column.SetShortName(LOCTEXT("TotalSizeOtherDependenciesColumnName", "Other"));
-		Column.SetTitleName(LOCTEXT("TotalSizeOtherDependenciesColumnTitle", "Total Other Dependency Size"));
-		Column.SetDescription(LOCTEXT("TotalSizeOtherDependenciesColumnIdDesc", "Sum of the staged compressed sizes of all dependencies of this asset which are shared by other assets directly or indirectly, counted only once"));
+		Column.SetShortName(LOCTEXT("TotalSizeSharedDependenciesColumnName", "Shared"));
+		Column.SetTitleName(LOCTEXT("TotalSizeSharedDependenciesColumnTitle", "Total Shared Dependency Size"));
+		Column.SetDescription(LOCTEXT("TotalSizeSharedDependenciesColumnIdDesc", "Sum of the staged compressed sizes of all dependencies of this asset which are shared by other assets directly or indirectly, counted only once"));
 
 		Column.SetFlags(ETableColumnFlags::CanBeHidden | ETableColumnFlags::CanBeFiltered);
 
@@ -951,7 +959,14 @@ void FAssetTable::AddDefaultColumns()
 				if (Node.IsGroup())
 				{
 					const FTableTreeNode& NodePtr = static_cast<const FTableTreeNode&>(Node);
-					if (NodePtr.HasAggregatedValue(Column.GetId()))
+					if (NodePtr.GetRowId().HasValidIndex())
+					{
+						// This group node also has a row index, so it conceptually only represents a single row (i.e., it's a synthetic group created by the asset dependency grouping)
+						TSharedPtr<FTable> TablePtr = NodePtr.GetParentTable().Pin();
+						const FAssetTable& AssetTable = static_cast<const FAssetTable&>(*TablePtr);
+						return FTableCellValue(static_cast<int64>(AssetTable.GetAssetChecked(NodePtr.GetRowIndex()).GetOrComputeTotalSizeSharedDependencies(AssetTable, NodePtr.GetRowIndex())));
+					}
+					else if (NodePtr.HasAggregatedValue(Column.GetId()))
 					{
 						return NodePtr.GetAggregatedValue(Column.GetId());
 					}
@@ -960,7 +975,73 @@ void FAssetTable::AddDefaultColumns()
 				{
 					const FAssetTreeNode& TreeNode = static_cast<const FAssetTreeNode&>(Node);
 					const FAssetTableRow& Asset = TreeNode.GetAssetChecked();
-					return FTableCellValue(static_cast<int64>(Asset.GetOrComputeTotalSizeOtherDependencies(TreeNode.GetAssetTableChecked(), TreeNode.GetRowIndex())));
+					return FTableCellValue(static_cast<int64>(Asset.GetOrComputeTotalSizeSharedDependencies(TreeNode.GetAssetTableChecked(), TreeNode.GetRowIndex())));
+				}
+
+				return TOptional<FTableCellValue>();
+			}
+		};
+		TSharedRef<ITableCellValueGetter> Getter = MakeShared<FTotalSizeUniqueDependenciesValueGetter>();
+		Column.SetValueGetter(Getter);
+
+		TSharedRef<ITableCellValueFormatter> Formatter = MakeShared<FInt64ValueFormatterAsMemory>();
+		Column.SetValueFormatter(Formatter);
+
+		TSharedRef<ITableCellValueSorter> Sorter = MakeShared<FSorterByInt64Value>(ColumnRef);
+		Column.SetValueSorter(Sorter);
+		Column.SetInitialSortMode(EColumnSortMode::Descending);
+
+		//TSharedRef<IFilterValueConverter> Converter = MakeShared<FMemoryFilterValueConverter>();
+		//Column.SetValueConverter(Converter);
+
+		Column.SetAggregation(ETableColumnAggregation::None);
+
+		AddColumn(ColumnRef);
+	}
+	//////////////////////////////////////////////////
+	// Total Size of External Dependencies
+	{
+		TSharedRef<FTableColumn> ColumnRef = MakeShared<FTableColumn>(FAssetTableColumns::TotalSizeExternalDependenciesColumnId);
+		FTableColumn& Column = *ColumnRef;
+
+		Column.SetIndex(ColumnIndex++);
+
+		Column.SetShortName(LOCTEXT("TotalSizeExternalDependenciesColumnName", "External"));
+		Column.SetTitleName(LOCTEXT("TotalSizeExternalDependenciesColumnTitle", "Total External Dependency Size"));
+		Column.SetDescription(LOCTEXT("TotalSizeExternalDependenciesColumnIdDesc", "Sum of the staged compressed sizes of all dependencies of this asset which are shared by other assets directly or indirectly and exist OUTSIDE the GFP of this asset, counted only once"));
+
+		Column.SetFlags(ETableColumnFlags::CanBeHidden | ETableColumnFlags::CanBeFiltered);
+
+		Column.SetHorizontalAlignment(HAlign_Right);
+		Column.SetInitialWidth(100.0f);
+
+		Column.SetDataType(ETableCellDataType::Int64);
+
+		class FTotalSizeUniqueDependenciesValueGetter : public FTableCellValueGetter
+		{
+		public:
+			virtual const TOptional<FTableCellValue> GetValue(const FTableColumn& Column, const FBaseTreeNode& Node) const override
+			{
+				if (Node.IsGroup())
+				{
+					const FTableTreeNode& NodePtr = static_cast<const FTableTreeNode&>(Node);
+					if (NodePtr.GetRowId().HasValidIndex())
+					{
+						// This group node also has a row index, so it conceptually only represents a single row (i.e., it's a synthetic group created by the asset dependency grouping)
+						TSharedPtr<FTable> TablePtr = NodePtr.GetParentTable().Pin();
+						const FAssetTable& AssetTable = static_cast<const FAssetTable&>(*TablePtr);
+						return FTableCellValue(static_cast<int64>(AssetTable.GetAssetChecked(NodePtr.GetRowIndex()).GetOrComputeTotalSizeExternalDependencies(AssetTable, NodePtr.GetRowIndex())));
+					}
+					else if (NodePtr.HasAggregatedValue(Column.GetId()))
+					{
+						return NodePtr.GetAggregatedValue(Column.GetId());
+					}
+				}
+				else //if (Node->Is<FAssetTreeNode>())
+				{
+					const FAssetTreeNode& TreeNode = static_cast<const FAssetTreeNode&>(Node);
+					const FAssetTableRow& Asset = TreeNode.GetAssetChecked();
+					return FTableCellValue(static_cast<int64>(Asset.GetOrComputeTotalSizeExternalDependencies(TreeNode.GetAssetTableChecked(), TreeNode.GetRowIndex())));
 				}
 
 				return TOptional<FTableCellValue>();
@@ -1228,7 +1309,7 @@ void FAssetTable::AddDefaultColumns()
 		Iterator.RemoveCurrent();
 
 		const FAssetTableRow& Row = OwningTable.GetAssetChecked(CurrentIndex);
-		if (!RestrictToGFPs.Contains(Row.GetGameFeaturePlugin()) || AdditionalNodesToStopAt.Contains(CurrentIndex))
+		if (ShouldSkipDueToGFP(RestrictToGFPs, Row.GetGameFeaturePlugin()) || AdditionalNodesToStopAt.Contains(CurrentIndex))
 		{
 			// Don't traverse outside this plugin
 			continue;
@@ -1249,10 +1330,8 @@ void FAssetTable::AddDefaultColumns()
 }
 
 ////////////////////////////////////////////////////////////////////////////////////////////////////
-// Iteratively refines the dependency set into unique and shared sets. PreviouslyVisitedIndices is split in each pass, moving some elements into
-// OutExcludedIndices (preserving its original contents) and putting those that are still potentially uniquely owned by the ThisIndex asset into
-// OutIncrementallyRefinedUniqueIndices. Returns 'true' if another pass is required (i.e., if work was done).
-/*static*/ void FAssetTableRow::RefineDependencies(const TSet<int32>& PreviouslyVisitedIndices, const FAssetTable& OwningTable, const TSet<int32> RootIndices, const TSet<const TCHAR*>& RestrictToGFPs, TSet<int32>& OutIncrementallyRefinedUniqueIndices, TSet<int32>& OutExcludedIndices)
+
+/*static*/ void FAssetTableRow::RefineDependencies(const TSet<int32>& PreviouslyVisitedIndices, const FAssetTable& OwningTable, const TSet<int32> RootIndices, const TSet<const TCHAR*>& RestrictToGFPs, TSet<int32>& OutUniqueIndices, TSet<int32>& OutSharedIndices)
 {
 	// "visit" ThisIndex to seed the exploration
 	TSet<int32> IndicesToVisit;
@@ -1269,13 +1348,13 @@ void FAssetTable::AddDefaultColumns()
 		Iterator.RemoveCurrent();
 
 		const FAssetTableRow& Row = OwningTable.GetAssetChecked(CurrentIndex);
-		if (!RestrictToGFPs.Contains(Row.GetGameFeaturePlugin()))
+		if (ShouldSkipDueToGFP(RestrictToGFPs, Row.GetGameFeaturePlugin()))
 		{
 			// Don't traverse outside this plugin
 			continue;
 		}
 
-		if (OutExcludedIndices.Contains(CurrentIndex))
+		if (OutSharedIndices.Contains(CurrentIndex))
 		{
 			// We already know not to traverse into this. We'll handle excluding its transitive dependencies later.
 			continue;
@@ -1295,12 +1374,12 @@ void FAssetTable::AddDefaultColumns()
 			NewlyExcludedDependencies.Add(CurrentIndex);
 			continue;
 		}
-		OutIncrementallyRefinedUniqueIndices.Add(CurrentIndex);
+		OutUniqueIndices.Add(CurrentIndex);
 
 		for (int32 ChildIndex : OwningTable.GetAssetChecked(CurrentIndex).GetDependencies())
 		{
 			// Don't revisit nodes we've already visited and don't re-add ThisIndex to avoid loops (and to avoid counting ourself)
-			if (!OutIncrementallyRefinedUniqueIndices.Contains(ChildIndex) && !RootIndices.Contains(ChildIndex))
+			if (!OutUniqueIndices.Contains(ChildIndex) && !RootIndices.Contains(ChildIndex))
 			{
 				IndicesToVisit.Add(ChildIndex);
 			}
@@ -1310,8 +1389,8 @@ void FAssetTable::AddDefaultColumns()
 	// We now have a set of nodes that we know are not uniquely owned
 	// Anything they reference should also be removed from the set of unique dependencies
 	TSet<int32> NewlyExcludedDependenciesAndTheirTransitiveDependencies = GatherAllReachableNodes(NewlyExcludedDependencies.Array(), OwningTable, RootIndices, RestrictToGFPs);
-	OutExcludedIndices = OutExcludedIndices.Union(NewlyExcludedDependenciesAndTheirTransitiveDependencies);
-	OutIncrementallyRefinedUniqueIndices = OutIncrementallyRefinedUniqueIndices.Difference(NewlyExcludedDependenciesAndTheirTransitiveDependencies);
+	OutSharedIndices = OutSharedIndices.Union(NewlyExcludedDependenciesAndTheirTransitiveDependencies);
+	OutUniqueIndices = OutUniqueIndices.Difference(NewlyExcludedDependenciesAndTheirTransitiveDependencies);
 }
 
 ////////////////////////////////////////////////////////////////////////////////////////////////////
@@ -1332,12 +1411,11 @@ void FAssetTable::AddDefaultColumns()
 	TSet<int32> VisitedIndices = GatherAllReachableNodes(RootIndices.Array(), OwningTable, TSet<int32>(), RestrictToGFPs);
 
 	// Iteratively separate the graph of "all things referenced by ThisIndex, directly or indirectly"
-	// into "UniqueDependencies -- things referenced ONLY by ThisIndex and by other things themselves referenced ONLY by ThisIndex"
-	// and "OtherDependencies" -- things removed from the list of "all things referenced by ThisIndex" in order to identify UniqueDependencies
+	// into "UniqueDependencies -- things referenced ONLY by ThisIndex and by other things themselves referenced ONLY by ThisIndex" limited to the GFPs of the RootIndices
+	// and "SharedDependencies" -- things removed from the list of "all things referenced by ThisIndex" in order to identify UniqueDependencies, limited to the GFPs of the RootIndices
 	TSet<int32> UniqueDependencies;
-	TSet<int32> OtherDependencies;
-	int32 IterationCount = 0;
-	RefineDependencies(VisitedIndices, OwningTable, RootIndices, RestrictToGFPs, UniqueDependencies, OtherDependencies);
+	TSet<int32> SharedDependencies;
+	RefineDependencies(VisitedIndices, OwningTable, RootIndices, RestrictToGFPs, UniqueDependencies, SharedDependencies);
 
 	// If there's only one root provided, the dependencies shouldn't include it. If more than one root is provided,
 	// some roots might be included as dependencies of other roots
@@ -1346,7 +1424,7 @@ void FAssetTable::AddDefaultColumns()
 	// if there is more than one root index, remove any root indices from the dependency lists in order to ensure they aren't double counted
 	// (they should be accounted for as part of the self sizes of the selected assets)
 	UniqueDependencies = UniqueDependencies.Difference(RootIndices);
-	OtherDependencies = OtherDependencies.Difference(RootIndices);
+	SharedDependencies = SharedDependencies.Difference(RootIndices);
 
 	for (int32 Index : UniqueDependencies)
 	{
@@ -1357,28 +1435,28 @@ void FAssetTable::AddDefaultColumns()
 		*OutUniqueDependencies = UniqueDependencies;
 	}
 
-	// Now explore all the dependencies of OtherDependencies and gather up their sizes
+	// Now explore all the dependencies of SharedDependencies and gather up their sizes
 	// This is necessary because the process calling RefineDependencies doesn't produce a complete list of other dependencies,
 	// just a partial set. By exploring the dependencies of that partial set we can find all the things that were
 	// referenced by ThisIndex but also by some other asset outside the subgraph defined by ThisIndex and its UniqueDependencies
 	VisitedIndices.Empty();
 
-	while (OtherDependencies.Num() > 0)
+	while (SharedDependencies.Num() > 0)
 	{
-		TSet<int32>::TIterator Iterator = OtherDependencies.CreateIterator();
+		TSet<int32>::TIterator Iterator = SharedDependencies.CreateIterator();
 		int32 CurrentIndex = *Iterator;
 		Iterator.RemoveCurrent();
 
 		const FAssetTableRow& Row = OwningTable.GetAssetChecked(CurrentIndex);
 		VisitedIndices.Add(CurrentIndex);
-		if (!RestrictToGFPs.Contains(Row.GetGameFeaturePlugin()))
+		if (ShouldSkipDueToGFP(RestrictToGFPs, Row.GetGameFeaturePlugin()))
 		{
 			// Don't traverse outside this plugin
 			continue;
 		}
 
 		int64 DependencySize = OwningTable.GetAssetChecked(CurrentIndex).GetStagedCompressedSize();
-		Result.OtherDependenciesSize += DependencySize;
+		Result.SharedDependenciesSize += DependencySize;
 		if (OutSharedDependencies != nullptr)
 		{
 			OutSharedDependencies->Add(CurrentIndex);
@@ -1388,7 +1466,7 @@ void FAssetTable::AddDefaultColumns()
 		{
 			if (!VisitedIndices.Contains(ChildIndex) && !RootIndices.Contains(ChildIndex))
 			{
-				OtherDependencies.Add(ChildIndex);
+				SharedDependencies.Add(ChildIndex);
 			}
 		}
 	}
@@ -1398,6 +1476,56 @@ void FAssetTable::AddDefaultColumns()
 
 ////////////////////////////////////////////////////////////////////////////////////////////////////
 
+/*static*/ int64 FAssetTableRow::ComputeTotalSizeExternalDependencies(const FAssetTable& OwningTable, const TSet<int32>& StartingNodes, TSet<int32>* OutExternalDependencies/* = nullptr*/)
+{
+	UE::Insights::FStopwatch Stopwatch;
+	Stopwatch.Start();
+
+	TSet<int32> AllReachableNodes = GatherAllReachableNodes(StartingNodes.Array(), OwningTable, TSet<int32>{}, TSet<const TCHAR*>{});
+	TSet<const TCHAR*> SourceGFPs;
+	for (int32 RootIndex : StartingNodes)
+	{
+		SourceGFPs.Add(OwningTable.GetAsset(RootIndex)->GetGameFeaturePlugin());
+	}
+
+	int64 TotalSizeExternalDependencies = 0;
+	int32 NumDependenciesIncluded = 0;
+	for (int32 Index : AllReachableNodes)
+	{
+		const FAssetTableRow& Row = *OwningTable.GetAsset(Index);
+		// Only include EXTERNAL dependencies
+		if (!SourceGFPs.Contains(Row.GameFeaturePlugin))
+		{
+			if (OutExternalDependencies != nullptr)
+			{
+				OutExternalDependencies->Add(Index);
+			}
+			TotalSizeExternalDependencies += Row.StagedCompressedSize;
+			NumDependenciesIncluded++;
+		}
+	}
+
+	Stopwatch.Stop();
+	const double Duration = Stopwatch.GetAccumulatedTime();
+	if (Duration > 1.0)
+	{
+		if (StartingNodes.Num() == 1)
+		{
+			int32 RowIndex = *StartingNodes.CreateConstIterator();
+			const FAssetTableRow& Row = *OwningTable.GetAsset(RowIndex);
+			UE_LOG(LogInsights, Log, TEXT("ComputeTotalSizeExternalDependencies(%s%s) : %.3fs"), Row.Path, Row.Name, Duration);
+		}
+		else
+		{
+			UE_LOG(LogInsights, Log, TEXT("ComputeTotalSizeExternalDependencies(<%d Assets>) --> %d Dependencies : %.3fs"), StartingNodes.Num(), NumDependenciesIncluded, Duration);
+		}
+	}
+
+	return TotalSizeExternalDependencies;
+}
+
+////////////////////////////////////////////////////////////////////////////////////////////////////
+
 int64 FAssetTableRow::GetOrComputeTotalSizeUniqueDependencies(const FAssetTable& OwningTable, int32 ThisIndex) const
 {
 	if (TotalSizeUniqueDependencies == -1)
@@ -1407,7 +1535,7 @@ int64 FAssetTableRow::GetOrComputeTotalSizeUniqueDependencies(const FAssetTable&
 
 		FAssetTableDependencySizes Sizes = ComputeDependencySizes(OwningTable, TSet<int32>{ThisIndex});
 		TotalSizeUniqueDependencies = Sizes.UniqueDependenciesSize;
-		TotalSizeOtherDependencies = Sizes.OtherDependenciesSize;
+		TotalSizeSharedDependencies = Sizes.SharedDependenciesSize;
 
 		Stopwatch.Stop();
 		const double Duration = Stopwatch.GetAccumulatedTime();
@@ -1421,16 +1549,16 @@ int64 FAssetTableRow::GetOrComputeTotalSizeUniqueDependencies(const FAssetTable&
 
 ////////////////////////////////////////////////////////////////////////////////////////////////////
 
-int64 FAssetTableRow::GetOrComputeTotalSizeOtherDependencies(const FAssetTable& OwningTable, int32 ThisIndex) const
+int64 FAssetTableRow::GetOrComputeTotalSizeSharedDependencies(const FAssetTable& OwningTable, int32 ThisIndex) const
 {
-	if (TotalSizeOtherDependencies == -1)
+	if (TotalSizeSharedDependencies == -1)
 	{
 		UE::Insights::FStopwatch Stopwatch;
 		Stopwatch.Start();
 
 		FAssetTableDependencySizes Sizes = ComputeDependencySizes(OwningTable, TSet<int32>{ThisIndex});
 		TotalSizeUniqueDependencies = Sizes.UniqueDependenciesSize;
-		TotalSizeOtherDependencies = Sizes.OtherDependenciesSize;
+		TotalSizeSharedDependencies = Sizes.SharedDependenciesSize;
 
 		Stopwatch.Stop();
 		const double Duration = Stopwatch.GetAccumulatedTime();
@@ -1439,7 +1567,18 @@ int64 FAssetTableRow::GetOrComputeTotalSizeOtherDependencies(const FAssetTable&
 			UE_LOG(LogInsights, Log, TEXT("ComputeDependencySizes(%s%s) : %.3fs"), Path, Name, Duration);
 		}
 	}
-	return TotalSizeOtherDependencies;
+	return TotalSizeSharedDependencies;
+}
+
+////////////////////////////////////////////////////////////////////////////////////////////////////
+
+int64 FAssetTableRow::GetOrComputeTotalSizeExternalDependencies(const FAssetTable& OwningTable, int32 ThisIndex) const
+{
+	if (TotalSizeExternalDependencies == -1)
+	{
+		TotalSizeExternalDependencies = ComputeTotalSizeExternalDependencies(OwningTable, TSet<int32>{ThisIndex});
+	}
+	return TotalSizeExternalDependencies;
 }
 
 ////////////////////////////////////////////////////////////////////////////////////////////////////
diff --git a/Engine/Plugins/Editor/AssetManagerEditor/Source/AssetManagerEditor/Private/TreeView/AssetTable.h b/Engine/Plugins/Editor/AssetManagerEditor/Source/AssetManagerEditor/Private/TreeView/AssetTable.h
--- a/Engine/Plugins/Editor/AssetManagerEditor/Source/AssetManagerEditor/Private/TreeView/AssetTable.h
+++ b/Engine/Plugins/Editor/AssetManagerEditor/Source/AssetManagerEditor/Private/TreeView/AssetTable.h
@@ -118,7 +118,8 @@ struct FAssetTableColumns
 	static const FName PrimaryNameColumnId;
 	static const FName StagedCompressedSizeColumnId;
 	static const FName TotalSizeUniqueDependenciesColumnId;
-	static const FName TotalSizeOtherDependenciesColumnId;
+	static const FName TotalSizeSharedDependenciesColumnId;
+	static const FName TotalSizeExternalDependenciesColumnId;
 	static const FName TotalUsageCountColumnId;
 	static const FName ChunksColumnId;
 	static const FName NativeClassColumnId;
@@ -128,7 +129,7 @@ struct FAssetTableColumns
 struct FAssetTableDependencySizes
 {
 	int64 UniqueDependenciesSize = 0;
-	int64 OtherDependenciesSize = 0;
+	int64 SharedDependenciesSize = 0;
 };
 
 class FAssetTableRow
@@ -153,7 +154,8 @@ public:
 	const TCHAR* GetPrimaryName() const { return PrimaryName; }
 	int64 GetStagedCompressedSize() const { return StagedCompressedSize; }
 	int64 GetOrComputeTotalSizeUniqueDependencies(const FAssetTable& OwningTable, int32 ThisIndex) const;
-	int64 GetOrComputeTotalSizeOtherDependencies(const FAssetTable& OwningTable, int32 ThisIndex) const; 
+	int64 GetOrComputeTotalSizeSharedDependencies(const FAssetTable& OwningTable, int32 ThisIndex) const; 
+	int64 GetOrComputeTotalSizeExternalDependencies(const FAssetTable& OwningTable, int32 ThisIndex) const; 
 	int64 GetTotalUsageCount() const { return TotalUsageCount; }
 	const TCHAR* GetChunks() const { return Chunks; }
 	const TCHAR* GetNativeClass() const { return NativeClass; }
@@ -166,8 +168,11 @@ public:
 
 	static FAssetTableDependencySizes ComputeDependencySizes(const FAssetTable& OwningTable, const TSet<int32>& RootIndices, TSet<int32>* OutUniqueDependencies = nullptr, TSet<int32>* OutSharedDependencies = nullptr);
 
+	static int64 ComputeTotalSizeExternalDependencies(const FAssetTable& OwningTable, const TSet<int32>& StartingNodes, TSet<int32>* OutExternalDependencies = nullptr);
+
 private:
-	static void RefineDependencies(const TSet<int32>& PreviouslyVisitedIndices, const FAssetTable& OwningTable, const TSet<int32> RootIndices, const TSet<const TCHAR*>& RestrictToGFPs, TSet<int32>& OutIncrementallyRefinedUniqueIndices, TSet<int32>& OutExcludedIndices);
+	static void RefineDependencies(const TSet<int32>& PreviouslyVisitedIndices, const FAssetTable& OwningTable, const TSet<int32> RootIndices, const TSet<const TCHAR*>& RestrictToGFPs, TSet<int32>& OutUniqueIndices, TSet<int32>& OutSharedIndices);
+	static inline bool ShouldSkipDueToGFP(const TSet<const TCHAR*>& RestrictToGFPs, const TCHAR* GFP) { return (RestrictToGFPs.Num() > 0) && !RestrictToGFPs.Contains(GFP); }
 
 	// Finds all nodes reachable from StartingNodes INCLUDING those in StartingNodes
 	static TSet<int32> GatherAllReachableNodes(const TArray<int32>& StartingNodes, const FAssetTable& OwningTable, const TSet<int32>& AdditionalNodesToStopAt, const TSet<const TCHAR*>& RestrictToGFPs);
@@ -180,7 +185,8 @@ private:
 	const TCHAR* PrimaryName = nullptr;
 	int64 StagedCompressedSize = 0;
 	mutable int64 TotalSizeUniqueDependencies = -1; // Lazily calculated
-	mutable int64 TotalSizeOtherDependencies = -1; // Lazily calculated
+	mutable int64 TotalSizeSharedDependencies = -1; // Lazily calculated
+	mutable int64 TotalSizeExternalDependencies = -1; // Lazily calculated
 	int64 TotalUsageCount = 0;
 	const TCHAR* Chunks = nullptr;
 	const TCHAR* NativeClass = nullptr;
diff --git a/Engine/Plugins/Editor/AssetManagerEditor/Source/AssetManagerEditor/Private/TreeView/SAssetTableTreeView.cpp b/Engine/Plugins/Editor/AssetManagerEditor/Source/AssetManagerEditor/Private/TreeView/SAssetTableTreeView.cpp
--- a/Engine/Plugins/Editor/AssetManagerEditor/Source/AssetManagerEditor/Private/TreeView/SAssetTableTreeView.cpp
+++ b/Engine/Plugins/Editor/AssetManagerEditor/Source/AssetManagerEditor/Private/TreeView/SAssetTableTreeView.cpp
@@ -14,6 +14,7 @@
 #include "Widgets/SToolTip.h"
 #include "DesktopPlatformModule.h"
 #include "HAL/PlatformFileManager.h"
+#include "HAL/FileManager.h"
 #include "Logging/MessageLog.h"
 
 #include "Insights/Common/InsightsStyle.h"
@@ -125,30 +126,12 @@ void SAssetTableTreeView::RebuildTree(bool bResync)
 
 	if (bResync || TableTreeNodes.Num() != PreviousNodeCount)
 	{
-		// Save selection.
-		TArray<UE::Insights::FTableTreeNodePtr> SelectedItems;
-		TreeView->GetSelectedItems(SelectedItems);
-
 		UE_LOG(LogInsights, Log, TEXT("[AssetTree] Update tree..."));
 		UpdateTree();
-
 		TreeView->RebuildList();
-
-		// Restore selection.
-		if (SelectedItems.Num() > 0)
-		{
-			TreeView->ClearSelection();
-			for (UE::Insights::FTableTreeNodePtr& NodePtr : SelectedItems)
-			{
-				NodePtr = GetNodeByTableRowIndex(NodePtr->GetRowIndex());
-			}
-			SelectedItems.RemoveAll([](const UE::Insights::FTableTreeNodePtr& NodePtr) { return !NodePtr.IsValid(); });
-			if (SelectedItems.Num() > 0)
-			{
-				TreeView->SetItemSelection(SelectedItems, true);
-				TreeView->RequestScrollIntoView(SelectedItems.Last());
-			}
-		}
+		TreeView->ClearSelection();
+		SelectedIndices.Empty();
+		SelectedAssetNode.Reset();
 	}
 
 	Stopwatch.Stop();
@@ -262,7 +245,8 @@ void SAssetTableTreeView::InitAvailableViewPresets()
 			InOutConfigSet.Add({ FAssetTableColumns::PrimaryNameColumnId,                   true, 200.0f });
 			InOutConfigSet.Add({ FAssetTableColumns::StagedCompressedSizeColumnId,          true, 100.0f });
 			InOutConfigSet.Add({ FAssetTableColumns::TotalSizeUniqueDependenciesColumnId,  !true, 100.0f });
-			InOutConfigSet.Add({ FAssetTableColumns::TotalSizeOtherDependenciesColumnId,   !true, 100.0f });
+			InOutConfigSet.Add({ FAssetTableColumns::TotalSizeSharedDependenciesColumnId,   !true, 100.0f });
+			InOutConfigSet.Add({ FAssetTableColumns::TotalSizeExternalDependenciesColumnId,!true, 100.0f });
 			InOutConfigSet.Add({ FAssetTableColumns::TotalUsageCountColumnId,               true, 100.0f });
 			InOutConfigSet.Add({ FAssetTableColumns::ChunksColumnId,                       !true, 200.0f });
 			InOutConfigSet.Add({ FAssetTableColumns::NativeClassColumnId,                   true, 200.0f });
@@ -272,9 +256,9 @@ void SAssetTableTreeView::InitAvailableViewPresets()
 	AvailableViewPresets.Add(MakeShared<FDefaultViewPreset>());
 
 	//////////////////////////////////////////////////
-	// GameFeaturePlugin, PrimaryType, Dependency View
+	// GameFeaturePlugin, Type, Dependency View
 	
-	class FGameFeaturePluginPrimaryTypeDependencyView : public UE::Insights::ITableTreeViewPreset
+	class FGameFeaturePluginTypeDependencyView : public UE::Insights::ITableTreeViewPreset
 	{
 	public:
 		virtual FText GetName() const override
@@ -284,7 +268,7 @@ void SAssetTableTreeView::InitAvailableViewPresets()
 
 		virtual FText GetToolTip() const override
 		{
-			return LOCTEXT("GFPTypeDepView_PresetToolTip", "Dependency Analysis View\nConfigure the tree view to show a breakdown of assets by Game Feature Plugin, Primary Asset Type, and Dependencies.");
+			return LOCTEXT("GFPTypeDepView_PresetToolTip", "Dependency Analysis View\nConfigure the tree view to show a breakdown of assets by Game Feature Plugin, Type, and Dependencies.");
 		}
 		virtual FName GetSortColumn() const override
 		{
@@ -316,7 +300,7 @@ void SAssetTableTreeView::InitAvailableViewPresets()
 				[](TSharedPtr<UE::Insights::FTreeNodeGrouping>& Grouping)
 				{
 					return Grouping->Is<UE::Insights::FTreeNodeGroupingByUniqueValueCString>() &&
-						Grouping->As<UE::Insights::FTreeNodeGroupingByUniqueValueCString>().GetColumnId() == FAssetTableColumns::PrimaryTypeColumnId;
+						Grouping->As<UE::Insights::FTreeNodeGroupingByUniqueValueCString>().GetColumnId() == FAssetTableColumns::TypeColumnId;
 				});
 			if (PrimaryTypeGrouping)
 			{
@@ -344,14 +328,15 @@ void SAssetTableTreeView::InitAvailableViewPresets()
 			InOutConfigSet.Add({ FAssetTableColumns::PrimaryNameColumnId,                  !true, 200.0f });
 			InOutConfigSet.Add({ FAssetTableColumns::StagedCompressedSizeColumnId,          true, 100.0f });
 			InOutConfigSet.Add({ FAssetTableColumns::TotalSizeUniqueDependenciesColumnId,   true, 100.0f });
-			InOutConfigSet.Add({ FAssetTableColumns::TotalSizeOtherDependenciesColumnId,    true, 100.0f });
+			InOutConfigSet.Add({ FAssetTableColumns::TotalSizeSharedDependenciesColumnId,    true, 100.0f });
+			InOutConfigSet.Add({ FAssetTableColumns::TotalSizeExternalDependenciesColumnId,!true, 100.0f });
 			InOutConfigSet.Add({ FAssetTableColumns::TotalUsageCountColumnId,              !true, 100.0f });
 			InOutConfigSet.Add({ FAssetTableColumns::ChunksColumnId,                       !true, 200.0f });
 			InOutConfigSet.Add({ FAssetTableColumns::NativeClassColumnId,                   true, 200.0f });
 			InOutConfigSet.Add({ FAssetTableColumns::GameFeaturePluginColumnId,             true, 200.0f });
 		}
 	};
-	AvailableViewPresets.Add(MakeShared<FGameFeaturePluginPrimaryTypeDependencyView>());
+	AvailableViewPresets.Add(MakeShared<FGameFeaturePluginTypeDependencyView>());
 
 	//////////////////////////////////////////////////
 	// Path Breakdown View
@@ -404,7 +389,8 @@ void SAssetTableTreeView::InitAvailableViewPresets()
 			InOutConfigSet.Add({ FAssetTableColumns::PrimaryNameColumnId,                   true, 200.0f });
 			InOutConfigSet.Add({ FAssetTableColumns::StagedCompressedSizeColumnId,          true, 100.0f });
 			InOutConfigSet.Add({ FAssetTableColumns::TotalSizeUniqueDependenciesColumnId,   true, 100.0f });
-			InOutConfigSet.Add({ FAssetTableColumns::TotalSizeOtherDependenciesColumnId,    true, 100.0f });
+			InOutConfigSet.Add({ FAssetTableColumns::TotalSizeSharedDependenciesColumnId,    true, 100.0f });
+			InOutConfigSet.Add({ FAssetTableColumns::TotalSizeExternalDependenciesColumnId,!true, 100.0f });
 			InOutConfigSet.Add({ FAssetTableColumns::TotalUsageCountColumnId,               true, 100.0f });
 			InOutConfigSet.Add({ FAssetTableColumns::ChunksColumnId,                       !true, 200.0f });
 			InOutConfigSet.Add({ FAssetTableColumns::NativeClassColumnId,                   true, 200.0f });
@@ -464,7 +450,8 @@ void SAssetTableTreeView::InitAvailableViewPresets()
 			InOutConfigSet.Add({ FAssetTableColumns::PrimaryNameColumnId,                  !true, 200.0f });
 			InOutConfigSet.Add({ FAssetTableColumns::StagedCompressedSizeColumnId,          true, 100.0f });
 			InOutConfigSet.Add({ FAssetTableColumns::TotalSizeUniqueDependenciesColumnId,   true, 100.0f });
-			InOutConfigSet.Add({ FAssetTableColumns::TotalSizeOtherDependenciesColumnId,    true, 100.0f });
+			InOutConfigSet.Add({ FAssetTableColumns::TotalSizeSharedDependenciesColumnId,    true, 100.0f });
+			InOutConfigSet.Add({ FAssetTableColumns::TotalSizeExternalDependenciesColumnId,!true, 100.0f });
 			InOutConfigSet.Add({ FAssetTableColumns::TotalUsageCountColumnId,               true, 100.0f });
 			InOutConfigSet.Add({ FAssetTableColumns::ChunksColumnId,                       !true, 200.0f });
 			InOutConfigSet.Add({ FAssetTableColumns::NativeClassColumnId,                   true, 200.0f });
@@ -576,8 +563,127 @@ void SAssetTableTreeView::AddCustomAdvancedFilters()
 
 ////////////////////////////////////////////////////////////////////////////////////////////////////
 
+void SAssetTableTreeView::ExportDependencyData() const
+{
+	IDesktopPlatform* DesktopPlatform = FDesktopPlatformModule::Get();
+	if (DesktopPlatform)
+	{
+		TSharedPtr<SWindow> ParentWindow = FSlateApplication::Get().FindWidgetWindow(AsShared());
+		void* ParentWindowHandle = (ParentWindow.IsValid() && ParentWindow->GetNativeWindow().IsValid()) ? ParentWindow->GetNativeWindow()->GetOSWindowHandle() : nullptr;
+
+		FString DefaultFileName;
+		if (SelectedIndices.Num() > 1)
+		{
+			DefaultFileName = "Batch Dependency Export.csv";
+		}
+		else
+		{
+			int32 RootIndex = *SelectedIndices.CreateConstIterator();
+			DefaultFileName = FString::Printf(TEXT("%s Dependencies.csv"), GetAssetTable()->GetAssetChecked(RootIndex).GetName());
+		}
+
+		TArray<FString> SaveFileNames;
+		const bool bFileSelected = DesktopPlatform->SaveFileDialog(
+			ParentWindowHandle,
+			LOCTEXT("ExportDependencyData_SaveFileDialogTitle", "Save dependency data as...").ToString(),
+			FPaths::ProjectLogDir(),
+			DefaultFileName,
+			TEXT("Comma-Separated Values (*.csv)|*.csv"),
+			EFileDialogFlags::None,
+			SaveFileNames);
+		
+		if (!bFileSelected)
+		{
+			return;
+		}
+		ensure(SaveFileNames.Num() == 1);
+		FString OutputFileName = SaveFileNames[0];
+
+
+		//FString FolderPath;
+		//const bool bFolderSelected = DesktopPlatform->OpenDirectoryDialog(
+		//	ParentWindowHandle,
+		//	LOCTEXT("FolderDialogTitle", "Choose a directory").ToString(),
+		//	FPaths::ProjectLogDir(),
+		//	FolderPath
+		//);
+
+		//if (bFolderSelected)
+		//{
+		//	if (!FolderPath.EndsWith(TEXT("/")))
+		//	{
+		//		FolderPath += TEXT("/");
+		//	}
+		//}
+
+		TSet<int32> UniqueDependencies;
+		TSet<int32> SharedDependencies;
+		FAssetTableRow::ComputeDependencySizes(*GetAssetTable(), SelectedIndices, &UniqueDependencies, &SharedDependencies);
+
+		TSet<int32> ExternalDependencies;
+		FAssetTableRow::ComputeTotalSizeExternalDependencies(*GetAssetTable(), SelectedIndices, &ExternalDependencies);
+
+		FString TimeSuffix = FDateTime::Now().ToString(TEXT("%Y%m%d_%H%M%S"));
+
+		TAnsiStringBuilder<4096> StringBuilder;
+		StringBuilder.Appendf("Asset,Self Size,Dependency Type\n");
+		{
+			for (int32 RootIndex : SelectedIndices)
+			{
+				const FAssetTableRow& Row = GetAssetTable()->GetAssetChecked(RootIndex);
+				StringBuilder.Appendf("%s%s,%s,Root\n", *WriteToAnsiString<512>(Row.GetPath()), *WriteToAnsiString<64>(Row.GetName()), *WriteToAnsiString<32>(LexToString(Row.GetStagedCompressedSize())));
+			}
+
+			for (int32 DependencyIndex : UniqueDependencies)
+			{
+				const FAssetTableRow& DependencyRow = GetAssetTable()->GetAssetChecked(DependencyIndex);
+				StringBuilder.Appendf("%s%s,%s,Unique\n", *WriteToAnsiString<512>(DependencyRow.GetPath()), *WriteToAnsiString<64>(DependencyRow.GetName()), *WriteToAnsiString<32>(LexToString(DependencyRow.GetStagedCompressedSize())));
+			}
+
+			for (int32 DependencyIndex : SharedDependencies)
+			{
+				const FAssetTableRow& DependencyRow = GetAssetTable()->GetAssetChecked(DependencyIndex);
+				StringBuilder.Appendf("%s%s,%s,Shared\n", *WriteToAnsiString<512>(DependencyRow.GetPath()), *WriteToAnsiString<64>(DependencyRow.GetName()), *WriteToAnsiString<32>(LexToString(DependencyRow.GetStagedCompressedSize())));
+			}
+
+			for (int32 DependencyIndex : ExternalDependencies)
+			{
+				const FAssetTableRow& DependencyRow = GetAssetTable()->GetAssetChecked(DependencyIndex);
+				StringBuilder.Appendf("%s%s,%s,External\n", *WriteToAnsiString<512>(DependencyRow.GetPath()), *WriteToAnsiString<64>(DependencyRow.GetName()), *WriteToAnsiString<32>(LexToString(DependencyRow.GetStagedCompressedSize())));
+			}
+		}
+
+		TUniquePtr<FArchive> DependencyFile(IFileManager::Get().CreateFileWriter(*OutputFileName));
+		DependencyFile->Serialize(StringBuilder.GetData(), StringBuilder.Len());
+	}
+}
+
 void SAssetTableTreeView::ExtendMenu(FMenuBuilder& MenuBuilder)
 {
+	MenuBuilder.BeginSection("ExportDependencies", LOCTEXT("ContextMenu_Section_ExportDependencies", "Export Dependencies"));
+	{
+		FUIAction ExportDependenciesAction;
+		ExportDependenciesAction.CanExecuteAction = FCanExecuteAction::CreateLambda([this]()
+			{
+				return GetAssetTable() && (SelectedIndices.Num() > 0);
+			});
+
+		ExportDependenciesAction.ExecuteAction = FExecuteAction::CreateLambda([this]()
+			{
+				ExportDependencyData();
+			});
+
+		MenuBuilder.AddMenuEntry
+		(
+			LOCTEXT("ContextMenu_ExportDependenciesLabel", "Export Dependencies..."),
+			LOCTEXT("ContextMenu_ExportDependencies", "Export dependency CSVs for the selected asset"),
+			FSlateIcon(),
+			ExportDependenciesAction,
+			NAME_None,
+			EUserInterfaceActionType::Button
+		);
+	}
+	MenuBuilder.EndSection();
 }
 
 ////////////////////////////////////////////////////////////////////////////////////////////////////
@@ -616,6 +722,7 @@ void SAssetTableTreeView::TreeView_OnSelectionChanged(UE::Insights::FTableTreeNo
 	const int32 NumSelectedNodes = TreeView->GetSelectedItems(SelectedNodes);
 	int32 NumSelectedAssets = 0;
 	FAssetTreeNodePtr NewSelectedAssetNode;
+	int32 NewlySelectedAssetRowIndex = -1;
 	TSet<int32> SelectionSetIndices;
 
 	for (const UE::Insights::FTableTreeNodePtr& Node : SelectedNodes)
@@ -623,9 +730,20 @@ void SAssetTableTreeView::TreeView_OnSelectionChanged(UE::Insights::FTableTreeNo
 		if (Node->Is<FAssetTreeNode>() && !Node->IsGroup())
 		{
 			NewSelectedAssetNode = StaticCastSharedPtr<FAssetTreeNode>(Node);
-			SelectionSetIndices.Add(NewSelectedAssetNode->GetRowId().RowIndex);
+			NewlySelectedAssetRowIndex = NewSelectedAssetNode->GetRowIndex();
+			SelectionSetIndices.Add(NewlySelectedAssetRowIndex);
 			++NumSelectedAssets;
 		}
+		else if (Node->Is<UE::Insights::FCustomTableTreeNode>() && Node->IsGroup())
+		{
+			int32 RowIndex = Node->GetRowIndex();
+			if (RowIndex != UE::Insights::FTableRowId::InvalidRowIndex)
+			{
+				SelectionSetIndices.Add(RowIndex);
+				NewlySelectedAssetRowIndex = RowIndex;
+				++NumSelectedAssets;
+			}
+		}
 	}
 
 	const int32 VisibleAssetCount = GetAssetTable()->GetVisibleAssetCount();
@@ -640,13 +758,14 @@ void SAssetTableTreeView::TreeView_OnSelectionChanged(UE::Insights::FTableTreeNo
 	{
 		FooterLeftText = FText::Format(LOCTEXT("FooterLeftTextFmt1", "{0} assets (1 selected)"), FText::AsNumber(VisibleAssetCount));
 
-		const FAssetTableRow& AssetTableRow = NewSelectedAssetNode->GetAssetChecked();
+		const FAssetTableRow& AssetTableRow = GetAssetTable()->GetAssetChecked(NewlySelectedAssetRowIndex);
 		FooterCenterText1 = FText::FromString(AssetTableRow.GetPath());
 		FooterCenterText2 = FText::FromString(AssetTableRow.GetName());
-		FooterRightText1 = FText::Format(LOCTEXT("FooterRightFmt3", "Self: {0} Unique: {1} Shared: {2}"),
+		FooterRightText1 = FText::Format(LOCTEXT("FooterRightFmt3", "Self: {0} Unique: {1} Shared: {2} External: {3}"),
 			FText::AsMemory(AssetTableRow.GetStagedCompressedSize()),
-			FText::AsMemory(AssetTableRow.GetOrComputeTotalSizeUniqueDependencies(*GetAssetTable(), NewSelectedAssetNode->GetRowId().RowIndex)),
-			FText::AsMemory(AssetTableRow.GetOrComputeTotalSizeOtherDependencies(*GetAssetTable(), NewSelectedAssetNode->GetRowId().RowIndex)));
+			FText::AsMemory(AssetTableRow.GetOrComputeTotalSizeUniqueDependencies(*GetAssetTable(), NewlySelectedAssetRowIndex)),
+			FText::AsMemory(AssetTableRow.GetOrComputeTotalSizeSharedDependencies(*GetAssetTable(), NewlySelectedAssetRowIndex)),
+			FText::AsMemory(AssetTableRow.GetOrComputeTotalSizeExternalDependencies(*GetAssetTable(), NewlySelectedAssetRowIndex)));
 	}
 	else
 	{
@@ -659,11 +778,14 @@ void SAssetTableTreeView::TreeView_OnSelectionChanged(UE::Insights::FTableTreeNo
 		{
 			TotalSelfSize += GetAssetTable()->GetAssetChecked(Index).GetStagedCompressedSize();
 		}
+
+		int64 TotalExternalDependencySize = FAssetTableRow::ComputeTotalSizeExternalDependencies(*GetAssetTable(), SelectionSetIndices);
 		FAssetTableDependencySizes Sizes = FAssetTableRow::ComputeDependencySizes(*GetAssetTable(), SelectionSetIndices, nullptr, nullptr);
-		FooterRightText1 = FText::Format(LOCTEXT("FooterRightFmt3", "Self: {0} Unique: {1} Shared: {2}"),
+		FooterRightText1 = FText::Format(LOCTEXT("FooterRightFmt3", "Self: {0} Unique: {1} Shared: {2} External: {3}"),
 			FText::AsMemory(TotalSelfSize),
 			FText::AsMemory(Sizes.UniqueDependenciesSize),
-			FText::AsMemory(Sizes.OtherDependenciesSize));
+			FText::AsMemory(Sizes.SharedDependenciesSize),
+			FText::AsMemory(TotalExternalDependencySize));
 	}
 
 	if (NumSelectedAssets != 1)
@@ -676,6 +798,8 @@ void SAssetTableTreeView::TreeView_OnSelectionChanged(UE::Insights::FTableTreeNo
 		SelectedAssetNode = NewSelectedAssetNode;
 	}
 
+	SelectedIndices = SelectionSetIndices;
+
 	if (OnSelectionChanged.IsBound())
 	{
 		OnSelectionChanged.Execute(SelectedNodes);
diff --git a/Engine/Plugins/Editor/AssetManagerEditor/Source/AssetManagerEditor/Private/TreeView/SAssetTableTreeView.h b/Engine/Plugins/Editor/AssetManagerEditor/Source/AssetManagerEditor/Private/TreeView/SAssetTableTreeView.h
--- a/Engine/Plugins/Editor/AssetManagerEditor/Source/AssetManagerEditor/Private/TreeView/SAssetTableTreeView.h
+++ b/Engine/Plugins/Editor/AssetManagerEditor/Source/AssetManagerEditor/Private/TreeView/SAssetTableTreeView.h
@@ -89,6 +89,8 @@ private:
 
 	void InitAvailableViewPresets();
 
+	void ExportDependencyData() const;
+
 private:
 	bool bNeedsToRebuild = false;
 
@@ -100,4 +102,5 @@ private:
 	/** Delegate to invoke when selection changes. */
 	FOnSelectionChanged OnSelectionChanged;
 	TSharedPtr<FAssetTreeNode> SelectedAssetNode;
+	TSet<int32> SelectedIndices;
 };
