Adding context menu options to asset registry audit tool to open selection in content browser, reference viewer, and editors

diff --git a/Engine/Plugins/Editor/AssetManagerEditor/Source/AssetManagerEditor/Private/SAssetAuditBrowser.cpp b/Engine/Plugins/Editor/AssetManagerEditor/Source/AssetManagerEditor/Private/SAssetAuditBrowser.cpp
--- a/Engine/Plugins/Editor/AssetManagerEditor/Source/AssetManagerEditor/Private/SAssetAuditBrowser.cpp
+++ b/Engine/Plugins/Editor/AssetManagerEditor/Source/AssetManagerEditor/Private/SAssetAuditBrowser.cpp
@@ -736,6 +736,8 @@ void SAssetAuditBrowser::GoToHistoryIndex(int32 InHistoryIdx)
 
 void SAssetAuditBrowser::QuickPopulateAssetTableRow(FAssetTableRow& OutRow, const FAssetData& AssetData, FAssetTable& AssetTable) const
 {
+	OutRow.SoftObjectPath = AssetData.GetSoftObjectPath();
+
 	// Asset Type
 	FString AssetType;
 	if (UAssetDefinitionRegistry* AssetDefinitionRegistry = UAssetDefinitionRegistry::Get())
diff --git a/Engine/Plugins/Editor/AssetManagerEditor/Source/AssetManagerEditor/Private/TreeView/AssetTable.h b/Engine/Plugins/Editor/AssetManagerEditor/Source/AssetManagerEditor/Private/TreeView/AssetTable.h
--- a/Engine/Plugins/Editor/AssetManagerEditor/Source/AssetManagerEditor/Private/TreeView/AssetTable.h
+++ b/Engine/Plugins/Editor/AssetManagerEditor/Source/AssetManagerEditor/Private/TreeView/AssetTable.h
@@ -5,7 +5,7 @@
 #include "CoreMinimal.h"
 #include "Containers/StringView.h"
 #include "Misc/CString.h"
-
+#include "AssetRegistry/AssetData.h"
 #include "Insights/Table/ViewModels/Table.h"
 
 namespace UE
@@ -165,6 +165,7 @@ public:
 	int32 GetNumReferencers() const { return Referencers.Num(); }
 	const TArray<int32>& GetReferencers() const { return Referencers; }
 	FLinearColor GetColor() const { return Color; }
+	const FSoftObjectPath& GetSoftObjectPath() const { return SoftObjectPath; }
 
 	static FAssetTableDependencySizes ComputeDependencySizes(const FAssetTable& OwningTable, const TSet<int32>& RootIndices, TSet<int32>* OutUniqueDependencies = nullptr, TSet<int32>* OutSharedDependencies = nullptr);
 
@@ -178,6 +179,7 @@ private:
 	static TSet<int32> GatherAllReachableNodes(const TArray<int32>& StartingNodes, const FAssetTable& OwningTable, const TSet<int32>& AdditionalNodesToStopAt, const TSet<const TCHAR*, TStringPointerSetKeyFuncs_DEPRECATED<const TCHAR*>>& RestrictToGFPs);
 
 private:
+	FSoftObjectPath SoftObjectPath; // Used to allow you to go to that asset in the content browser or to open its editor after lookingup the FAssetData in the original registry state
 	const TCHAR* Type = nullptr;
 	const TCHAR* Name = nullptr;
 	const TCHAR* Path = nullptr;
diff --git a/Engine/Plugins/Editor/AssetManagerEditor/Source/AssetManagerEditor/Private/TreeView/SAssetTableTreeView.cpp b/Engine/Plugins/Editor/AssetManagerEditor/Source/AssetManagerEditor/Private/TreeView/SAssetTableTreeView.cpp
--- a/Engine/Plugins/Editor/AssetManagerEditor/Source/AssetManagerEditor/Private/TreeView/SAssetTableTreeView.cpp
+++ b/Engine/Plugins/Editor/AssetManagerEditor/Source/AssetManagerEditor/Private/TreeView/SAssetTableTreeView.cpp
@@ -3,6 +3,7 @@
 #include "SAssetTableTreeView.h"
 
 #include "Containers/Set.h"
+#include "ContentBrowserModule.h"
 #include "Framework/MultiBox/MultiBoxBuilder.h"
 #include "GenericPlatform/GenericPlatformMisc.h"
 #include "Modules/ModuleManager.h"
@@ -16,6 +17,11 @@
 #include "HAL/PlatformFileManager.h"
 #include "HAL/FileManager.h"
 #include "Logging/MessageLog.h"
+#include "Editor.h"
+#include "Subsystems/AssetEditorSubsystem.h"
+#include "IContentBrowserSingleton.h"
+#include "AssetManagerEditorModule.h"
+#include "AssetRegistry/AssetRegistryModule.h"
 
 #include "Insights/Common/InsightsStyle.h"
 #include "Insights/Common/Log.h"
@@ -33,6 +39,8 @@
 
 #define LOCTEXT_NAMESPACE "SAssetTableTreeView"
 
+extern UNREALED_API UEditorEngine* GEditor;
+
 ////////////////////////////////////////////////////////////////////////////////////////////////////
 
 SAssetTableTreeView::SAssetTableTreeView()
@@ -658,15 +666,119 @@ void SAssetTableTreeView::ExportDependencyData() const
 	}
 }
 
+TArray<FAssetData> SAssetTableTreeView::GetAssetDataForSelection() const
+{
+	TArray<FAssetData> Assets;
+	for (int32 SelectionIndex : SelectedIndices)
+	{
+		const FAssetManagerEditorRegistrySource* RegistrySource = IAssetManagerEditorModule::Get().GetCurrentRegistrySource();
+		const FSoftObjectPath& SoftObjectPath = GetAssetTable()->GetAssetChecked(SelectionIndex).GetSoftObjectPath();
+		FAssetData AssetData;
+		if (RegistrySource->GetOwnedRegistryState())
+		{
+			AssetData = *RegistrySource->GetOwnedRegistryState()->GetAssetByObjectPath(SoftObjectPath);
+		}
+		else
+		{
+			// Once we move this out of asset audit browser window and give it its own registry state, this code path can be killed
+			// it exists to handle the "Editor" case as a registry source
+			AssetData = FModuleManager::LoadModuleChecked<FAssetRegistryModule>(TEXT("AssetRegistry")).Get().GetAssetByObjectPath(SoftObjectPath, /*IncludeOnlyOnDiskAsset*/true, /*SkipARFilteredAssets*/false);
+		}
+		Assets.Add(AssetData);
+	}
+
+	return Assets;
+}
+
 void SAssetTableTreeView::ExtendMenu(FMenuBuilder& MenuBuilder)
 {
-	MenuBuilder.BeginSection("ExportDependencies", LOCTEXT("ContextMenu_Section_ExportDependencies", "Export Dependencies"));
+	FCanExecuteAction HasSelectionAndCanExecute = FCanExecuteAction::CreateLambda([this]()
+		{
+			return GetAssetTable() && (SelectedIndices.Num() > 0);
+		});
+
+	FCanExecuteAction HasSelectionAndRegistrySourceAndCanExecute = FCanExecuteAction::CreateLambda([this]()
+		{
+			return (IAssetManagerEditorModule::Get().GetCurrentRegistrySource() != nullptr) && GetAssetTable() && (SelectedIndices.Num() > 0);
+		});
+
+	MenuBuilder.BeginSection("Asset", LOCTEXT("ContextMenu_Section_Asset", "Asset"));
 	{
-		FUIAction ExportDependenciesAction;
-		ExportDependenciesAction.CanExecuteAction = FCanExecuteAction::CreateLambda([this]()
+		//////////////////////////////////////////////////////////////////////////
+		/// EditSelectedAssets
+
+		FUIAction EditSelectedAssets;
+		EditSelectedAssets.CanExecuteAction = HasSelectionAndCanExecute;
+		EditSelectedAssets.ExecuteAction = FExecuteAction::CreateLambda([this]()
+			{
+				TArray<FSoftObjectPath> AssetPaths;
+				for (int32 SelectionIndex : SelectedIndices)
+				{
+					AssetPaths.Add(GetAssetTable()->GetAssetChecked(SelectionIndex).GetSoftObjectPath());
+				}
+
+				GEditor->GetEditorSubsystem<UAssetEditorSubsystem>()->OpenEditorsForAssets(AssetPaths);
+			});
+		MenuBuilder.AddMenuEntry
+			(
+				LOCTEXT("ContextMenu_EditAssetsLabel", "Edit..."),
+				LOCTEXT("ContextMenu_EditAssets", "Opens the selected aset in the relevant editor."),
+				FSlateIcon(),
+				EditSelectedAssets,
+				NAME_None,
+				EUserInterfaceActionType::Button
+			);
+
+		//////////////////////////////////////////////////////////////////////////
+		/// FindInContentBrowser
+
+		FUIAction FindInContentBrowser;
+		FindInContentBrowser.CanExecuteAction = HasSelectionAndRegistrySourceAndCanExecute;
+		FindInContentBrowser.ExecuteAction = FExecuteAction::CreateLambda([this]()
+			{
+				FContentBrowserModule& ContentBrowserModule = FModuleManager::Get().LoadModuleChecked<FContentBrowserModule>("ContentBrowser");
+				ContentBrowserModule.Get().SyncBrowserToAssets(GetAssetDataForSelection());
+		});
+		MenuBuilder.AddMenuEntry
+		(
+			LOCTEXT("ContextMenu_FindInContentBrowserLabel", "Find in Content Browser..."),
+			LOCTEXT("ContextMenu_FindInContentBrowser", "Browses to the associated asset and selects it in the most recently used Content Browser (summoning one if necessary)"),
+			FSlateIcon(),
+			FindInContentBrowser,
+			NAME_None,
+			EUserInterfaceActionType::Button
+		);
+
+		//////////////////////////////////////////////////////////////////////////
+		/// OpenReferenceViewer
+
+		FUIAction OpenReferenceViewer;
+		OpenReferenceViewer.CanExecuteAction = HasSelectionAndCanExecute;
+		OpenReferenceViewer.ExecuteAction = FExecuteAction::CreateLambda([this]()
 			{
-				return GetAssetTable() && (SelectedIndices.Num() > 0);
+				TArray<FAssetIdentifier> AssetIdentifiers;
+				IAssetManagerEditorModule::ExtractAssetIdentifiersFromAssetDataList(GetAssetDataForSelection(), AssetIdentifiers);
+
+				if (AssetIdentifiers.Num() > 0)
+				{
+					IAssetManagerEditorModule::Get().OpenReferenceViewerUI(AssetIdentifiers);
+				}
 			});
+		MenuBuilder.AddMenuEntry
+		(
+			LOCTEXT("ContextMenu_OpenReferenceViewerLabel", "Reference Viewer..."),
+			LOCTEXT("ContextMenu_OpenReferenceViewer", "Launches the reference viewer showing the selected assets' references"),
+			FSlateIcon(),
+			OpenReferenceViewer,
+			NAME_None,
+			EUserInterfaceActionType::Button
+		);
+
+		//////////////////////////////////////////////////////////////////////////
+		/// ExportDependencies
+
+		FUIAction ExportDependenciesAction;
+		ExportDependenciesAction.CanExecuteAction = HasSelectionAndCanExecute;
 
 		ExportDependenciesAction.ExecuteAction = FExecuteAction::CreateLambda([this]()
 			{
diff --git a/Engine/Plugins/Editor/AssetManagerEditor/Source/AssetManagerEditor/Private/TreeView/SAssetTableTreeView.h b/Engine/Plugins/Editor/AssetManagerEditor/Source/AssetManagerEditor/Private/TreeView/SAssetTableTreeView.h
--- a/Engine/Plugins/Editor/AssetManagerEditor/Source/AssetManagerEditor/Private/TreeView/SAssetTableTreeView.h
+++ b/Engine/Plugins/Editor/AssetManagerEditor/Source/AssetManagerEditor/Private/TreeView/SAssetTableTreeView.h
@@ -91,6 +91,8 @@ private:
 
 	void ExportDependencyData() const;
 
+	TArray<FAssetData> GetAssetDataForSelection() const;
+
 private:
 	bool bNeedsToRebuild = false;
 
