AssetManagerEditor:
  - Updated STableTreeView and FBaseTreeNode implementations to latest from TraceInsights module (node filtering / hierarchy filtering improvements).
  - Fixed aggregated value for Unique Size column to stop propagation at group nodes that have a valid asset index.

diff --git a/Engine/Plugins/Editor/AssetManagerEditor/Source/AssetManagerEditor/Private/Insights/Filter/ViewModels/FilterConfigurator.cpp b/Engine/Plugins/Editor/AssetManagerEditor/Source/AssetManagerEditor/Private/Insights/Filter/ViewModels/FilterConfigurator.cpp
--- a/Engine/Plugins/Editor/AssetManagerEditor/Source/AssetManagerEditor/Private/Insights/Filter/ViewModels/FilterConfigurator.cpp
+++ b/Engine/Plugins/Editor/AssetManagerEditor/Source/AssetManagerEditor/Private/Insights/Filter/ViewModels/FilterConfigurator.cpp
@@ -41,8 +41,8 @@ FFilterConfigurator& FFilterConfigurator::operator=(const FFilterConfigurator& O
 		return *this;
 	}
 
-	RootNode = MakeShared<FFilterConfiguratorNode>(*Other.RootNode);
-	RootNode->SetGroupPtrForChildren();
+	RootNode = FFilterConfiguratorNode::DeepCopy(*Other.RootNode);
+
 	AvailableFilters = Other.AvailableFilters;
 	OnDestroyedEvent = Other.OnDestroyedEvent;
 
diff --git a/Engine/Plugins/Editor/AssetManagerEditor/Source/AssetManagerEditor/Private/Insights/Filter/ViewModels/FilterConfigurator.h b/Engine/Plugins/Editor/AssetManagerEditor/Source/AssetManagerEditor/Private/Insights/Filter/ViewModels/FilterConfigurator.h
--- a/Engine/Plugins/Editor/AssetManagerEditor/Source/AssetManagerEditor/Private/Insights/Filter/ViewModels/FilterConfigurator.h
+++ b/Engine/Plugins/Editor/AssetManagerEditor/Source/AssetManagerEditor/Private/Insights/Filter/ViewModels/FilterConfigurator.h
@@ -23,18 +23,20 @@ public:
 	FFilterConfigurator& operator=(const FFilterConfigurator& Other);
 
 	bool operator==(const FFilterConfigurator& Other) const;
-
 	bool operator!=(const FFilterConfigurator& Other) const { return !(*this == Other); }
 
 	virtual ~FFilterConfigurator();
 
-	FFilterConfiguratorNodePtr GetRootNode() { return RootNode; }
+	bool IsEmpty() const { return RootNode->GetChildrenCount() == 0; }
 
 	virtual bool ApplyFilters(const FFilterContext& Context) const override;
 
 	bool IsKeyUsed(int32 Key) const;
 
-	TSharedPtr<TArray<TSharedPtr<struct FFilter>>>& GetAvailableFilters() { return AvailableFilters; }
+	FFilterConfiguratorNodePtr GetRootNode() { return RootNode; }
+	TSharedPtr<TArray<TSharedPtr<FFilter>>>& GetAvailableFilters() { return AvailableFilters; }
+
+	void Add(TSharedPtr<FFilter> InFilter) { AvailableFilters->Add(InFilter); }
 
 	////////////////////////////////////////////////////////////////////////////////////////////////////
 	// OnDestroyedEvent
@@ -66,7 +68,7 @@ private:
 private:
 	FFilterConfiguratorNodePtr RootNode;
 
-	TSharedPtr<TArray<TSharedPtr<struct FFilter>>> AvailableFilters;
+	TSharedPtr<TArray<TSharedPtr<FFilter>>> AvailableFilters;
 
 	TSet<int32> KeysUsed;
 };
diff --git a/Engine/Plugins/Editor/AssetManagerEditor/Source/AssetManagerEditor/Private/Insights/Filter/ViewModels/FilterConfiguratorNode.cpp b/Engine/Plugins/Editor/AssetManagerEditor/Source/AssetManagerEditor/Private/Insights/Filter/ViewModels/FilterConfiguratorNode.cpp
--- a/Engine/Plugins/Editor/AssetManagerEditor/Source/AssetManagerEditor/Private/Insights/Filter/ViewModels/FilterConfiguratorNode.cpp
+++ b/Engine/Plugins/Editor/AssetManagerEditor/Source/AssetManagerEditor/Private/Insights/Filter/ViewModels/FilterConfiguratorNode.cpp
@@ -7,7 +7,7 @@
 // Insights
 #include "Insights/Filter/ViewModels/Filters.h"
 
-#define LOCTEXT_NAMESPACE "FFilterConfiguratorNode"
+#define LOCTEXT_NAMESPACE "UE::Insights::FFilterConfiguratorNode"
 
 namespace UE
 {
@@ -31,32 +31,33 @@ FFilterConfiguratorNode::FFilterConfiguratorNode(const FName InName, bool bInIsG
 
 ////////////////////////////////////////////////////////////////////////////////////////////////////
 
-FFilterConfiguratorNode::FFilterConfiguratorNode(const FFilterConfiguratorNode& Other)
-	: FBaseTreeNode(Other.GetName(), Other.IsGroup())
+TSharedRef<FFilterConfiguratorNode> FFilterConfiguratorNode::DeepCopy(const FFilterConfiguratorNode& InSourceNode)
 {
-	*this = Other;
-}
+	TSharedRef<FFilterConfiguratorNode> NodeCopyPtr = MakeShared<FFilterConfiguratorNode>(InSourceNode.GetName(), InSourceNode.IsGroup());
+	FFilterConfiguratorNode& NodeCopy = *NodeCopyPtr;
 
-////////////////////////////////////////////////////////////////////////////////////////////////////
+	NodeCopy.AvailableFilters = InSourceNode.AvailableFilters;
+	NodeCopy.AvailableFilterOperators = InSourceNode.AvailableFilterOperators;
+	NodeCopy.SelectedFilter = InSourceNode.SelectedFilter;
+	NodeCopy.SelectedFilterOperator = InSourceNode.SelectedFilterOperator;
+	NodeCopy.SelectedFilterGroupOperator = InSourceNode.SelectedFilterGroupOperator;
+	NodeCopy.TextBoxValue = InSourceNode.TextBoxValue;
 
-FFilterConfiguratorNode& FFilterConfiguratorNode::operator=(const FFilterConfiguratorNode& Other)
-{
-	GetChildrenMutable().Empty();
+	NodeCopy.SetExpansion(InSourceNode.IsExpanded());
 
-	AvailableFilters = Other.AvailableFilters;
-	AvailableFilterOperators = Other.AvailableFilterOperators;
-	SelectedFilter = Other.SelectedFilter;
-	SelectedFilterOperator = Other.SelectedFilterOperator;
-	SelectedFilterGroupOperator = Other.SelectedFilterGroupOperator;
-	TextBoxValue = Other.TextBoxValue;
-	SetExpansion(Other.IsExpanded());
-
-	for (FBaseTreeNodePtr Child : Other.GetChildren())
+	if (InSourceNode.IsGroup())
 	{
-		GetChildrenMutable().Add(MakeShared<FFilterConfiguratorNode>(*StaticCastSharedPtr<FFilterConfiguratorNode>(Child)));
+		check(NodeCopy.IsGroup());
+		NodeCopy.ClearChildren(InSourceNode.GetChildrenCount());
+		for (FBaseTreeNodePtr Child : InSourceNode.GetChildren())
+		{
+			check(Child.IsValid() && Child->Is<FFilterConfiguratorNode>());
+			TSharedPtr<FFilterConfiguratorNode> ChildCopy = FFilterConfiguratorNode::DeepCopy(Child->As<FFilterConfiguratorNode>());
+			NodeCopy.AddChildAndSetParent(ChildCopy);
+		}
 	}
 
-	return *this;
+	return NodeCopyPtr;
 }
 
 ////////////////////////////////////////////////////////////////////////////////////////////////////
@@ -69,11 +70,12 @@ bool FFilterConfiguratorNode::operator==(const FFilterConfiguratorNode& Other) c
 	bIsEqual &= SelectedFilterOperator.Get() == Other.SelectedFilterOperator.Get();
 	bIsEqual &= SelectedFilterGroupOperator.Get() == Other.SelectedFilterGroupOperator.Get();
 	bIsEqual &= TextBoxValue == Other.TextBoxValue;
-	bIsEqual &= GetChildren().Num() == Other.GetChildren().Num();
+	bIsEqual &= GetChildrenCount() == Other.GetChildrenCount();
 
 	if (bIsEqual)
 	{
-		for (int32 Index = 0; Index < GetChildren().Num(); ++Index)
+		const int32 Count = GetChildrenCount();
+		for (int32 Index = 0; Index < Count; ++Index)
 		{
 			bIsEqual &= *StaticCastSharedPtr<FFilterConfiguratorNode>(GetChildren()[Index]) == *StaticCastSharedPtr<FFilterConfiguratorNode>(Other.GetChildren()[Index]);
 		}
@@ -84,34 +86,14 @@ bool FFilterConfiguratorNode::operator==(const FFilterConfiguratorNode& Other) c
 
 ////////////////////////////////////////////////////////////////////////////////////////////////////
 
-const TArray<TSharedPtr<struct FFilterGroupOperator>>& FFilterConfiguratorNode::GetFilterGroupOperators()
+const TArray<TSharedPtr<FFilterGroupOperator>>& FFilterConfiguratorNode::GetFilterGroupOperators()
 {
 	return FFilterService::Get()->GetFilterGroupOperators();
 }
 
 ////////////////////////////////////////////////////////////////////////////////////////////////////
 
-void FFilterConfiguratorNode::DeleteChildNode(FFilterConfiguratorNodePtr InNode)
-{
-	FBaseTreeNodePtr Node = StaticCastSharedPtr<FBaseTreeNode>(InNode);
-	GetChildrenMutable().Remove(Node);
-}
-
-////////////////////////////////////////////////////////////////////////////////////////////////////
-
-void FFilterConfiguratorNode::SetGroupPtrForChildren()
-{
-	for (FBaseTreeNodePtr Child : GetChildrenMutable())
-	{
-		FFilterConfiguratorNodePtr CastedChild = StaticCastSharedPtr<FFilterConfiguratorNode>(Child);
-		CastedChild->SetGroupPtrForChildren();
-		CastedChild->SetGroupPtr(SharedThis(this));
-	}
-}
-
-////////////////////////////////////////////////////////////////////////////////////////////////////
-
-void FFilterConfiguratorNode::SetAvailableFilters(TSharedPtr<TArray<TSharedPtr<struct FFilter>>> InAvailableFilters)
+void FFilterConfiguratorNode::SetAvailableFilters(TSharedPtr<TArray<TSharedPtr<FFilter>>> InAvailableFilters)
 {
 	AvailableFilters = InAvailableFilters;
 
@@ -123,7 +105,7 @@ void FFilterConfiguratorNode::SetAvailableFilters(TSharedPtr<TArray<TSharedPtr<s
 
 ////////////////////////////////////////////////////////////////////////////////////////////////////
 
-void FFilterConfiguratorNode::SetSelectedFilter(TSharedPtr<struct FFilter> InSelectedFilter)
+void FFilterConfiguratorNode::SetSelectedFilter(TSharedPtr<FFilter> InSelectedFilter)
 {
 	SelectedFilter = InSelectedFilter;
 	if (SelectedFilter.IsValid() && SelectedFilter->GetSupportedOperators()->Num() > 0)
@@ -131,7 +113,7 @@ void FFilterConfiguratorNode::SetSelectedFilter(TSharedPtr<struct FFilter> InSel
 		SetSelectedFilterOperator(SelectedFilter->GetSupportedOperators()->GetData()[0]);
 
 		AvailableFilterOperators->Empty();
-		SupportedOperatorsArrayPtr AvailableOperators = InSelectedFilter->SupportedOperators;
+		SupportedOperatorsArrayPtr AvailableOperators = InSelectedFilter->GetSupportedOperators();
 		for (auto& FilterOperator : *AvailableOperators)
 		{
 			AvailableFilterOperators->Add(FilterOperator);
@@ -154,15 +136,15 @@ void FFilterConfiguratorNode::ProcessFilter()
 	}
 	else
 	{
-		switch (SelectedFilter->DataType)
+		switch (SelectedFilter->GetDataType())
 		{
 		case EFilterDataType::Double:
 		{
-			if (SelectedFilter->Converter.IsValid())
+			if (SelectedFilter->GetConverter().IsValid())
 			{
 				double Value = 0.0;
 				FText Errors;
-				bool Result = SelectedFilter->Converter->Convert(TextBoxValue, Value, Errors);
+				bool Result = SelectedFilter->GetConverter()->Convert(TextBoxValue, Value, Errors);
 				FilterValue.Set<double>(Result ? Value : 0.0);
 			}
 			else
@@ -173,11 +155,11 @@ void FFilterConfiguratorNode::ProcessFilter()
 		}
 		case EFilterDataType::Int64:
 		{
-			if (SelectedFilter->Converter.IsValid())
+			if (SelectedFilter->GetConverter().IsValid())
 			{
 				int64 Value = 0;
 				FText Errors;
-				bool Result = SelectedFilter->Converter->Convert(TextBoxValue, Value, Errors);
+				bool Result = SelectedFilter->GetConverter()->Convert(TextBoxValue, Value, Errors);
 				FilterValue.Set<int64>(Result ? Value : 0);
 			}
 			else
@@ -200,10 +182,10 @@ void FFilterConfiguratorNode::ProcessFilter()
 		}
 		case EFilterDataType::StringInt64Pair:
 		{
-			checkf(SelectedFilter->Converter.IsValid(), TEXT("StringToInt64Pair filters must have a converter set"));
+			checkf(SelectedFilter->GetConverter().IsValid(), TEXT("StringToInt64Pair filters must have a converter set"));
 			int64 Value = 0;
 			FText Errors;
-			bool Result = SelectedFilter->Converter->Convert(TextBoxValue, Value, Errors);
+			bool Result = SelectedFilter->GetConverter()->Convert(TextBoxValue, Value, Errors);
 			FilterValue.Set<int64>(Result ? Value : -1);
 		}
 		}
@@ -222,7 +204,7 @@ void FFilterConfiguratorNode::GetUsedKeys(TSet<int32>& KeysUsed) const
 	}
 	else
 	{
-		KeysUsed.Add(SelectedFilter->Key);
+		KeysUsed.Add(SelectedFilter->GetKey());
 	}
 }
 
@@ -233,7 +215,7 @@ bool FFilterConfiguratorNode::ApplyFilters(const FFilterContext& Context) const
 	bool Ret = true;
 	if (IsGroup())
 	{
-		switch (SelectedFilterGroupOperator->Type)
+		switch (SelectedFilterGroupOperator->GetType())
 		{
 		case EFilterGroupOperator::And:
 		{
@@ -262,21 +244,21 @@ bool FFilterConfiguratorNode::ApplyFilters(const FFilterContext& Context) const
 	}
 	else
 	{
-		if (!Context.HasFilterData(SelectedFilter->Key))
+		if (!Context.HasFilterData(SelectedFilter->GetKey()))
 		{
 			// If data is not set for this filter return the value specified in the Context.
 			return Context.GetReturnValueForUnsetFilters();
 		}
 
-		switch (SelectedFilter->DataType)
+		switch (SelectedFilter->GetDataType())
 		{
 		case EFilterDataType::Double:
 		{
 			FFilterOperator<double>* Operator = (FFilterOperator<double>*) SelectedFilterOperator.Get();
 			double Value = 0.0;
-			Context.GetFilterData<double>(SelectedFilter->Key, Value);
+			Context.GetFilterData<double>(SelectedFilter->GetKey(), Value);
 
-			Ret = Operator->Func(Value, FilterValue.Get<double>());
+			Ret = Operator->Apply(Value, FilterValue.Get<double>());
 			break;
 		}
 		case EFilterDataType::Int64:
@@ -284,18 +266,18 @@ bool FFilterConfiguratorNode::ApplyFilters(const FFilterContext& Context) const
 		{
 			FFilterOperator<int64>* Operator = (FFilterOperator<int64>*) SelectedFilterOperator.Get();
 			int64 Value = 0;
-			Context.GetFilterData<int64>(SelectedFilter->Key, Value);
+			Context.GetFilterData<int64>(SelectedFilter->GetKey(), Value);
 
-			Ret = Operator->Func(Value, FilterValue.Get<int64>());
+			Ret = Operator->Apply(Value, FilterValue.Get<int64>());
 			break;
 		}
 		case EFilterDataType::String:
 		{
 			FFilterOperator<FString>* Operator = (FFilterOperator<FString>*) SelectedFilterOperator.Get();
 			FString Value;
-			Context.GetFilterData<FString>(SelectedFilter->Key, Value);
+			Context.GetFilterData<FString>(SelectedFilter->GetKey(), Value);
 
-			Ret = Operator->Func(Value, FilterValue.Get<FString>());
+			Ret = Operator->Apply(Value, FilterValue.Get<FString>());
 			break;
 		}
 		default:
diff --git a/Engine/Plugins/Editor/AssetManagerEditor/Source/AssetManagerEditor/Private/Insights/Filter/ViewModels/FilterConfiguratorNode.h b/Engine/Plugins/Editor/AssetManagerEditor/Source/AssetManagerEditor/Private/Insights/Filter/ViewModels/FilterConfiguratorNode.h
--- a/Engine/Plugins/Editor/AssetManagerEditor/Source/AssetManagerEditor/Private/Insights/Filter/ViewModels/FilterConfiguratorNode.h
+++ b/Engine/Plugins/Editor/AssetManagerEditor/Source/AssetManagerEditor/Private/Insights/Filter/ViewModels/FilterConfiguratorNode.h
@@ -53,34 +53,30 @@ public:
 	/** Initialization constructor for the filter configurator node. */
 	FFilterConfiguratorNode(const FName InName, bool bInIsGroup);
 
-	FFilterConfiguratorNode(const FFilterConfiguratorNode& Other);
+	static TSharedRef<FFilterConfiguratorNode> DeepCopy(const FFilterConfiguratorNode& Node);
 
 	bool operator==(const FFilterConfiguratorNode& Other) const;
 
 	virtual ~FFilterConfiguratorNode() {}
 
-	void SetAvailableFilters(TSharedPtr<TArray<TSharedPtr<struct FFilter>>> InAvailableFilters);
-	TSharedPtr<TArray<TSharedPtr<struct FFilter>>> GetAvailableFilters() { return AvailableFilters; }
+	void SetAvailableFilters(TSharedPtr<TArray<TSharedPtr<FFilter>>> InAvailableFilters);
+	TSharedPtr<TArray<TSharedPtr<FFilter>>> GetAvailableFilters() { return AvailableFilters; }
 
-	void SetSelectedFilter(TSharedPtr<struct FFilter> InSelectedFilter);
-	TSharedPtr<struct FFilter> GetSelectedFilter() const { return SelectedFilter; }
+	void SetSelectedFilter(TSharedPtr<FFilter> InSelectedFilter);
+	TSharedPtr<FFilter> GetSelectedFilter() const { return SelectedFilter; }
 
-	void SetSelectedFilterOperator(TSharedPtr<class IFilterOperator> InSelectedFilterOperator) { SelectedFilterOperator = InSelectedFilterOperator; }
+	void SetSelectedFilterOperator(TSharedPtr<IFilterOperator> InSelectedFilterOperator) { SelectedFilterOperator = InSelectedFilterOperator; }
 	TSharedPtr<IFilterOperator> GetSelectedFilterOperator() const { return SelectedFilterOperator;	}
-	const TArray<TSharedPtr<struct FFilterGroupOperator>>& GetFilterGroupOperators();
+	const TArray<TSharedPtr<FFilterGroupOperator>>& GetFilterGroupOperators();
 
-	void SetSelectedFilterGroupOperator(TSharedPtr<struct FFilterGroupOperator> InSelectedFilterGroupOperator) { SelectedFilterGroupOperator = InSelectedFilterGroupOperator; }
-	TSharedPtr<struct FFilterGroupOperator> GetSelectedFilterGroupOperator() const { return SelectedFilterGroupOperator; }
+	void SetSelectedFilterGroupOperator(TSharedPtr<FFilterGroupOperator> InSelectedFilterGroupOperator) { SelectedFilterGroupOperator = InSelectedFilterGroupOperator; }
+	TSharedPtr<FFilterGroupOperator> GetSelectedFilterGroupOperator() const { return SelectedFilterGroupOperator; }
 
 	TSharedPtr<TArray<TSharedPtr<IFilterOperator>>> GetAvailableFilterOperators() const {	return AvailableFilterOperators;	}
 
-	void DeleteChildNode(FFilterConfiguratorNodePtr InNode);
-
 	const FString& GetTextBoxValue() { return TextBoxValue; }
 	void SetTextBoxValue(const FString& InValue) { TextBoxValue = InValue; }
 
-	void SetGroupPtrForChildren();
-
 	bool ApplyFilters(const class FFilterContext& Context) const;
 
 	void GetUsedKeys(TSet<int32>& GetUsedKeys) const;
@@ -90,13 +86,13 @@ public:
 private:
 	FFilterConfiguratorNode& operator=(const FFilterConfiguratorNode& Other);
 
-	TSharedPtr<TArray<TSharedPtr<struct FFilter>>> AvailableFilters;
+	TSharedPtr<TArray<TSharedPtr<FFilter>>> AvailableFilters;
 
-	TSharedPtr<struct FFilter> SelectedFilter;
+	TSharedPtr<FFilter> SelectedFilter;
 
-	TSharedPtr<class IFilterOperator> SelectedFilterOperator;
+	TSharedPtr<IFilterOperator> SelectedFilterOperator;
 
-	TSharedPtr<struct FFilterGroupOperator> SelectedFilterGroupOperator;
+	TSharedPtr<FFilterGroupOperator> SelectedFilterGroupOperator;
 
 	TSharedPtr<TArray<TSharedPtr<IFilterOperator>>> AvailableFilterOperators;
 
diff --git a/Engine/Plugins/Editor/AssetManagerEditor/Source/AssetManagerEditor/Private/Insights/Filter/ViewModels/Filters.cpp b/Engine/Plugins/Editor/AssetManagerEditor/Source/AssetManagerEditor/Private/Insights/Filter/ViewModels/Filters.cpp
--- a/Engine/Plugins/Editor/AssetManagerEditor/Source/AssetManagerEditor/Private/Insights/Filter/ViewModels/Filters.cpp
+++ b/Engine/Plugins/Editor/AssetManagerEditor/Source/AssetManagerEditor/Private/Insights/Filter/ViewModels/Filters.cpp
@@ -10,7 +10,7 @@
 #include "Insights/Filter/ViewModels/FilterConfigurator.h"
 #include "Insights/Filter/Widgets/SAdvancedFilter.h"
 
-#define LOCTEXT_NAMESPACE "SFilterService"
+#define LOCTEXT_NAMESPACE "UE::Insights::Filters"
 
 namespace UE
 {
@@ -20,6 +20,44 @@ namespace Insights
 INSIGHTS_IMPLEMENT_RTTI(FFilter)
 INSIGHTS_IMPLEMENT_RTTI(FFilterWithSuggestions)
 
+////////////////////////////////////////////////////////////////////////////////////////////////////
+// FFilterStorage
+////////////////////////////////////////////////////////////////////////////////////////////////////
+
+FFilterStorage::FFilterStorage()
+{
+	DoubleOperators = MakeShared<TArray<TSharedPtr<IFilterOperator>>>();
+	DoubleOperators->Add(StaticCastSharedRef<IFilterOperator>(MakeShared<FFilterOperator<double>>(EFilterOperator::Lt, TEXT("<"), std::less<>{})));
+	DoubleOperators->Add(StaticCastSharedRef<IFilterOperator>(MakeShared<FFilterOperator<double>>(EFilterOperator::Lte, TEXT("\u2264"), std::less_equal<>())));
+	DoubleOperators->Add(StaticCastSharedRef<IFilterOperator>(MakeShared<FFilterOperator<double>>(EFilterOperator::Eq, TEXT("="), std::equal_to<>())));
+	DoubleOperators->Add(StaticCastSharedRef<IFilterOperator>(MakeShared<FFilterOperator<double>>(EFilterOperator::Gt, TEXT(">"), std::greater<>())));
+	DoubleOperators->Add(StaticCastSharedRef<IFilterOperator>(MakeShared<FFilterOperator<double>>(EFilterOperator::Gte, TEXT("\u2265"), std::greater_equal<>())));
+
+	IntegerOperators = MakeShared<TArray<TSharedPtr<IFilterOperator>>>();
+	IntegerOperators->Add(StaticCastSharedRef<IFilterOperator>(MakeShared<FFilterOperator<int64>>(EFilterOperator::Lt, TEXT("<"), std::less<>{})));
+	IntegerOperators->Add(StaticCastSharedRef<IFilterOperator>(MakeShared<FFilterOperator<int64>>(EFilterOperator::Lte, TEXT("\u2264"), std::less_equal<>())));
+	IntegerOperators->Add(StaticCastSharedRef<IFilterOperator>(MakeShared<FFilterOperator<int64>>(EFilterOperator::Eq, TEXT("="), std::equal_to<>())));
+	IntegerOperators->Add(StaticCastSharedRef<IFilterOperator>(MakeShared<FFilterOperator<int64>>(EFilterOperator::Gt, TEXT(">"), std::greater<>())));
+	IntegerOperators->Add(StaticCastSharedRef<IFilterOperator>(MakeShared<FFilterOperator<int64>>(EFilterOperator::Gte, TEXT("\u2265"), std::greater_equal<>())));
+
+	StringOperators = MakeShared<TArray<TSharedPtr<IFilterOperator>>>();
+	StringOperators->Add(StaticCastSharedRef<IFilterOperator>(MakeShared<FFilterOperator<FString>>(EFilterOperator::Eq, TEXT("IS"), [](const FString& lhs, const FString& rhs) { return lhs.Equals(rhs); })));
+	StringOperators->Add(StaticCastSharedRef<IFilterOperator>(MakeShared<FFilterOperator<FString>>(EFilterOperator::NotEq, TEXT("IS NOT"), [](const FString& lhs, const FString& rhs) { return !lhs.Equals(rhs); })));
+	StringOperators->Add(StaticCastSharedRef<IFilterOperator>(MakeShared<FFilterOperator<FString>>(EFilterOperator::Contains, TEXT("CONTAINS"), [](const FString& lhs, const FString& rhs) { return lhs.Contains(rhs); })));
+	StringOperators->Add(StaticCastSharedRef<IFilterOperator>(MakeShared<FFilterOperator<FString>>(EFilterOperator::NotContains, TEXT("NOT CONTAINS"), [](const FString& lhs, const FString& rhs) { return !lhs.Contains(rhs); })));
+
+	FilterGroupOperators.Add(MakeShared<FFilterGroupOperator>(EFilterGroupOperator::And, LOCTEXT("AllOf", "All Of (AND)"), LOCTEXT("AllOfDesc", "All of the children must be true for the group to return true. Equivalent to an AND operation.")));
+	FilterGroupOperators.Add(MakeShared<FFilterGroupOperator>(EFilterGroupOperator::Or, LOCTEXT("AnyOf", "Any Of (OR)"), LOCTEXT("AnyOfDesc", "Any of the children must be true for the group to return true. Equivalent to an OR operation.")));
+}
+
+////////////////////////////////////////////////////////////////////////////////////////////////////
+
+FFilterStorage::~FFilterStorage()
+{
+}
+
+////////////////////////////////////////////////////////////////////////////////////////////////////
+// FFilterService
 ////////////////////////////////////////////////////////////////////////////////////////////////////
 
 TSharedPtr<FFilterService> FFilterService::Instance;
diff --git a/Engine/Plugins/Editor/AssetManagerEditor/Source/AssetManagerEditor/Private/Insights/Filter/ViewModels/Filters.h b/Engine/Plugins/Editor/AssetManagerEditor/Source/AssetManagerEditor/Private/Insights/Filter/ViewModels/Filters.h
--- a/Engine/Plugins/Editor/AssetManagerEditor/Source/AssetManagerEditor/Private/Insights/Filter/ViewModels/Filters.h
+++ b/Engine/Plugins/Editor/AssetManagerEditor/Source/AssetManagerEditor/Private/Insights/Filter/ViewModels/Filters.h
@@ -9,8 +9,6 @@
 
 #include "Insights/Common/SimpleRtti.h"
 
-#define LOCTEXT_NAMESPACE "Filters"
-
 class FSpawnTabArgs;
 class SDockTab;
 class SWidget;
@@ -49,9 +47,8 @@ enum class EFilterOperator : uint8
 class IFilterOperator
 {
 public:
-
-	virtual EFilterOperator GetKey() = 0;
-	virtual FString GetName() = 0;
+	virtual EFilterOperator GetKey() const = 0;
+	virtual FString GetName() const = 0;
 };
 
 ////////////////////////////////////////////////////////////////////////////////////////////////////
@@ -60,26 +57,29 @@ template<typename T>
 class FFilterOperator : public IFilterOperator
 {
 public:
-	typedef TFunction<bool(T, T)> OperatorFunc;
+	typedef TFunction<bool(T, T)> FOperatorFunc;
 
-	FFilterOperator(EFilterOperator InKey, FString InName, OperatorFunc InFunc)
-		: Func(InFunc)
-		, Key(InKey)
+public:
+	FFilterOperator(EFilterOperator InKey, FString InName, FOperatorFunc InFunc)
+		: Key(InKey)
 		, Name(InName)
+		, Func(InFunc)
 	{
 	}
+
 	virtual ~FFilterOperator()
 	{
 	}
 
-	virtual EFilterOperator GetKey() override { return Key; }
-	virtual FString GetName() override { return Name; };
+	virtual EFilterOperator GetKey() const override { return Key; }
+	virtual FString GetName() const override { return Name; };
 
-	OperatorFunc Func;
+	bool Apply(T InValueA, T InValueB) const { return Func(InValueA, InValueB); }
 
 private:
 	EFilterOperator Key;
 	FString Name;
+	FOperatorFunc Func;
 };
 
 ////////////////////////////////////////////////////////////////////////////////////////////////////
@@ -92,8 +92,9 @@ enum class EFilterGroupOperator
 
 ////////////////////////////////////////////////////////////////////////////////////////////////////
 
-struct FFilterGroupOperator
+class FFilterGroupOperator
 {
+public:
 	FFilterGroupOperator(EFilterGroupOperator InType, FText InName, FText InDesc)
 		: Type(InType)
 		, Name(InName)
@@ -101,6 +102,11 @@ struct FFilterGroupOperator
 	{
 	}
 
+	EFilterGroupOperator GetType() const { return Type; }
+	const FText& GetName() const { return Name; }
+	const FText& GetDesc() const { return Desc; }
+
+private:
 	EFilterGroupOperator Type;
 	FText Name;
 	FText Desc;
@@ -121,25 +127,33 @@ public:
 
 typedef TSharedPtr<const TArray<TSharedPtr<IFilterOperator>>> SupportedOperatorsArrayPtr;
 
-struct FFilter
+class FFilter
 {
 	INSIGHTS_DECLARE_RTTI_BASE(FFilter)
 
 public:
-	FFilter(int32 InKey, FText InName, FText InDesc, EFilterDataType InDataType, SupportedOperatorsArrayPtr InSupportedOperators)
+	FFilter(int32 InKey, FText InName, FText InDesc, EFilterDataType InDataType, TSharedPtr<IFilterValueConverter> InConverter, SupportedOperatorsArrayPtr InSupportedOperators)
 		: Key(InKey)
 		, Name(InName)
 		, Desc(InDesc)
 		, DataType(InDataType)
+		, Converter(InConverter)
 		, SupportedOperators(InSupportedOperators)
 	{
 	}
+
 	virtual ~FFilter()
 	{
 	}
 
+	int32 GetKey() const { return Key; }
+	const FText& GetName() const { return Name; }
+	const FText& GetDesc() const { return Desc; }
+	EFilterDataType GetDataType() const { return DataType; }
+	const TSharedPtr<IFilterValueConverter>& GetConverter() const { return Converter; }
 	SupportedOperatorsArrayPtr GetSupportedOperators() const { return SupportedOperators; }
 
+private:
 	int32 Key;
 	FText Name;
 	FText Desc;
@@ -150,15 +164,16 @@ public:
 
 ////////////////////////////////////////////////////////////////////////////////////////////////////
 
-struct FFilterWithSuggestions : FFilter
+class FFilterWithSuggestions : public FFilter
 {
-	typedef TFunction<void(const FString& /*Text*/, TArray<FString>& OutSuggestions)> GetSuggestionsCallback;
-
 	INSIGHTS_DECLARE_RTTI(FFilterWithSuggestions, FFilter)
 
 public:
-	FFilterWithSuggestions(int32 InKey, FText InName, FText InDesc, EFilterDataType InDataType, SupportedOperatorsArrayPtr InSupportedOperators)
-		: FFilter(InKey, InName, InDesc, InDataType, InSupportedOperators)
+	typedef TFunction<void(const FString& /*Text*/, TArray<FString>& OutSuggestions)> FGetSuggestionsCallback;
+
+public:
+	FFilterWithSuggestions(int32 InKey, FText InName, FText InDesc, EFilterDataType InDataType, TSharedPtr<IFilterValueConverter> InConverter, SupportedOperatorsArrayPtr InSupportedOperators)
+		: FFilter(InKey, InName, InDesc, InDataType, InConverter, InSupportedOperators)
 	{
 	}
 
@@ -166,7 +181,12 @@ public:
 	{
 	}
 
-	GetSuggestionsCallback Callback;
+	const FGetSuggestionsCallback& GetCallback() const { return Callback; }
+	void SetCallback(FGetSuggestionsCallback InCallback) { Callback = InCallback; }
+	void GetSuggestions(const FString& Text, TArray<FString>& OutSuggestions) { Callback(Text, OutSuggestions); }
+
+private:
+	FGetSuggestionsCallback Callback;
 };
 
 ////////////////////////////////////////////////////////////////////////////////////////////////////
@@ -174,31 +194,8 @@ public:
 class FFilterStorage
 {
 public:
-	FFilterStorage()
-	{
-		DoubleOperators = MakeShared<TArray<TSharedPtr<IFilterOperator>>>();
-		DoubleOperators->Add(StaticCastSharedRef<IFilterOperator>(MakeShared<FFilterOperator<double>>(EFilterOperator::Lt, TEXT("<"), std::less<>{})));
-		DoubleOperators->Add(StaticCastSharedRef<IFilterOperator>(MakeShared<FFilterOperator<double>>(EFilterOperator::Lte, TEXT("\u2264"), std::less_equal<>())));
-		DoubleOperators->Add(StaticCastSharedRef<IFilterOperator>(MakeShared<FFilterOperator<double>>(EFilterOperator::Eq, TEXT("="), std::equal_to<>())));
-		DoubleOperators->Add(StaticCastSharedRef<IFilterOperator>(MakeShared<FFilterOperator<double>>(EFilterOperator::Gt, TEXT(">"), std::greater<>())));
-		DoubleOperators->Add(StaticCastSharedRef<IFilterOperator>(MakeShared<FFilterOperator<double>>(EFilterOperator::Gte, TEXT("\u2265"), std::greater_equal<>())));
-
-		IntegerOperators = MakeShared<TArray<TSharedPtr<IFilterOperator>>>();
-		IntegerOperators->Add(StaticCastSharedRef<IFilterOperator>(MakeShared<FFilterOperator<int64>>(EFilterOperator::Lt, TEXT("<"), std::less<>{})));
-		IntegerOperators->Add(StaticCastSharedRef<IFilterOperator>(MakeShared<FFilterOperator<int64>>(EFilterOperator::Lte, TEXT("\u2264"), std::less_equal<>())));
-		IntegerOperators->Add(StaticCastSharedRef<IFilterOperator>(MakeShared<FFilterOperator<int64>>(EFilterOperator::Eq, TEXT("="), std::equal_to<>())));
-		IntegerOperators->Add(StaticCastSharedRef<IFilterOperator>(MakeShared<FFilterOperator<int64>>(EFilterOperator::Gt, TEXT(">"), std::greater<>())));
-		IntegerOperators->Add(StaticCastSharedRef<IFilterOperator>(MakeShared<FFilterOperator<int64>>(EFilterOperator::Gte, TEXT("\u2265"), std::greater_equal<>())));
-
-		StringOperators = MakeShared<TArray<TSharedPtr<IFilterOperator>>>();
-		StringOperators->Add(StaticCastSharedRef<IFilterOperator>(MakeShared<FFilterOperator<FString>>(EFilterOperator::Eq, TEXT("IS"), [](const FString& lhs, const FString& rhs) { return lhs.Equals(rhs); })));
-		StringOperators->Add(StaticCastSharedRef<IFilterOperator>(MakeShared<FFilterOperator<FString>>(EFilterOperator::NotEq, TEXT("IS NOT"), [](const FString& lhs, const FString& rhs) { return !lhs.Equals(rhs); })));
-		StringOperators->Add(StaticCastSharedRef<IFilterOperator>(MakeShared<FFilterOperator<FString>>(EFilterOperator::Contains, TEXT("CONTAINS"), [](const FString& lhs, const FString& rhs) { return lhs.Contains(rhs); })));
-		StringOperators->Add(StaticCastSharedRef<IFilterOperator>(MakeShared<FFilterOperator<FString>>(EFilterOperator::NotContains, TEXT("NOT CONTAINS"), [](const FString& lhs, const FString& rhs) { return !lhs.Contains(rhs); })));
-
-		FilterGroupOperators.Add(MakeShared<FFilterGroupOperator>(EFilterGroupOperator::And, LOCTEXT("AllOf", "All Of (AND)"), LOCTEXT("AllOfDesc", "All of the children must be true for the group to return true. Equivalent to an AND operation.")));
-		FilterGroupOperators.Add(MakeShared<FFilterGroupOperator>(EFilterGroupOperator::Or, LOCTEXT("AnyOf", "Any Of (OR)"), LOCTEXT("AnyOfDesc", "Any of the children must be true for the group to return true. Equivalent to an OR operation.")));
-	}
+	FFilterStorage();
+	~FFilterStorage();
 
 	const TArray<TSharedPtr<FFilterGroupOperator>>& GetFilterGroupOperators() const { return FilterGroupOperators; }
 
@@ -254,6 +251,7 @@ class FFilterContext
 public:
 	typedef TVariant<double, int64, FString> ContextData;
 
+public:
 	template<typename T>
 	void AddFilterData(int32 Key, const T& InData)
 	{
@@ -295,5 +293,3 @@ private:
 
 } // namespace Insights
 } // namespace UE
-
-#undef LOCTEXT_NAMESPACE
diff --git a/Engine/Plugins/Editor/AssetManagerEditor/Source/AssetManagerEditor/Private/Insights/Filter/Widgets/SAdvancedFilter.cpp b/Engine/Plugins/Editor/AssetManagerEditor/Source/AssetManagerEditor/Private/Insights/Filter/Widgets/SAdvancedFilter.cpp
--- a/Engine/Plugins/Editor/AssetManagerEditor/Source/AssetManagerEditor/Private/Insights/Filter/Widgets/SAdvancedFilter.cpp
+++ b/Engine/Plugins/Editor/AssetManagerEditor/Source/AssetManagerEditor/Private/Insights/Filter/Widgets/SAdvancedFilter.cpp
@@ -17,7 +17,7 @@
 #include "Insights/Filter/Widgets/SFilterConfigurator.h"
 #include "Insights/Filter/ViewModels/FilterConfigurator.h"
 
-#define LOCTEXT_NAMESPACE "SAdvancedFilter"
+#define LOCTEXT_NAMESPACE "UE::Insights::SAdvancedFilter"
 
 namespace UE
 {
diff --git a/Engine/Plugins/Editor/AssetManagerEditor/Source/AssetManagerEditor/Private/Insights/Filter/Widgets/SAdvancedFilter.h b/Engine/Plugins/Editor/AssetManagerEditor/Source/AssetManagerEditor/Private/Insights/Filter/Widgets/SAdvancedFilter.h
--- a/Engine/Plugins/Editor/AssetManagerEditor/Source/AssetManagerEditor/Private/Insights/Filter/Widgets/SAdvancedFilter.h
+++ b/Engine/Plugins/Editor/AssetManagerEditor/Source/AssetManagerEditor/Private/Insights/Filter/Widgets/SAdvancedFilter.h
@@ -18,6 +18,7 @@ namespace UE
 namespace Insights
 {
 
+class FFilterConfigurator;
 class SFilterConfigurator;
 
 ////////////////////////////////////////////////////////////////////////////////////////////////////
@@ -40,7 +41,7 @@ public:
 	 * Construct this widget
 	 * @param InArgs - The declaration data for this widget
 	 */
-	void Construct(const FArguments& InArgs, TSharedPtr<class FFilterConfigurator> InFilterConfiguratorViewModel);
+	void Construct(const FArguments& InArgs, TSharedPtr<FFilterConfigurator> InFilterConfiguratorViewModel);
 
 	void Reset();
 
@@ -59,9 +60,9 @@ private:
 private:
 	TSharedPtr<SFilterConfigurator> FilterConfigurator;
 
-	TWeakPtr<class FFilterConfigurator> OriginalFilterConfiguratorViewModel;
+	TWeakPtr<FFilterConfigurator> OriginalFilterConfiguratorViewModel;
 
-	TSharedPtr<class FFilterConfigurator> FilterConfiguratorViewModel;
+	TSharedPtr<FFilterConfigurator> FilterConfiguratorViewModel;
 
 	TWeakPtr<SDockTab> ParentTab;
 
diff --git a/Engine/Plugins/Editor/AssetManagerEditor/Source/AssetManagerEditor/Private/Insights/Filter/Widgets/SFilterConfigurator.cpp b/Engine/Plugins/Editor/AssetManagerEditor/Source/AssetManagerEditor/Private/Insights/Filter/Widgets/SFilterConfigurator.cpp
--- a/Engine/Plugins/Editor/AssetManagerEditor/Source/AssetManagerEditor/Private/Insights/Filter/Widgets/SFilterConfigurator.cpp
+++ b/Engine/Plugins/Editor/AssetManagerEditor/Source/AssetManagerEditor/Private/Insights/Filter/Widgets/SFilterConfigurator.cpp
@@ -17,7 +17,7 @@
 #include "Insights/Filter/Widgets/SFilterConfiguratorRow.h"
 #include "Insights/Filter/ViewModels/FilterConfigurator.h"
 
-#define LOCTEXT_NAMESPACE "SFilterConfigurator"
+#define LOCTEXT_NAMESPACE "UE::Insights::SFilterConfigurator"
 
 namespace UE
 {
diff --git a/Engine/Plugins/Editor/AssetManagerEditor/Source/AssetManagerEditor/Private/Insights/Filter/Widgets/SFilterConfigurator.h b/Engine/Plugins/Editor/AssetManagerEditor/Source/AssetManagerEditor/Private/Insights/Filter/Widgets/SFilterConfigurator.h
--- a/Engine/Plugins/Editor/AssetManagerEditor/Source/AssetManagerEditor/Private/Insights/Filter/Widgets/SFilterConfigurator.h
+++ b/Engine/Plugins/Editor/AssetManagerEditor/Source/AssetManagerEditor/Private/Insights/Filter/Widgets/SFilterConfigurator.h
@@ -18,6 +18,8 @@ namespace UE
 namespace Insights
 {
 
+class FFilterConfigurator;
+
 ////////////////////////////////////////////////////////////////////////////////////////////////////
 
 /**
@@ -39,7 +41,7 @@ public:
 	 * Construct this widget
 	 * @param InArgs - The declaration data for this widget
 	 */
-	void Construct(const FArguments& InArgs, TSharedPtr<class FFilterConfigurator> InFilterConfiguratorViewModel);
+	void Construct(const FArguments& InArgs, TSharedPtr<FFilterConfigurator> InFilterConfiguratorViewModel);
 
 	void Reset();
 
@@ -64,7 +66,7 @@ private:
 	void SetInitialExpansionRec(const FFilterConfiguratorNodePtr& Node, bool Value);
 
 private:
-	TSharedPtr<class FFilterConfigurator> FilterConfiguratorViewModel;
+	TSharedPtr<FFilterConfigurator> FilterConfiguratorViewModel;
 
 	/** The tree widget which holds the filters. */
 	TSharedPtr<STreeView<FFilterConfiguratorNodePtr>> TreeView;
diff --git a/Engine/Plugins/Editor/AssetManagerEditor/Source/AssetManagerEditor/Private/Insights/Filter/Widgets/SFilterConfiguratorRow.cpp b/Engine/Plugins/Editor/AssetManagerEditor/Source/AssetManagerEditor/Private/Insights/Filter/Widgets/SFilterConfiguratorRow.cpp
--- a/Engine/Plugins/Editor/AssetManagerEditor/Source/AssetManagerEditor/Private/Insights/Filter/Widgets/SFilterConfiguratorRow.cpp
+++ b/Engine/Plugins/Editor/AssetManagerEditor/Source/AssetManagerEditor/Private/Insights/Filter/Widgets/SFilterConfiguratorRow.cpp
@@ -19,7 +19,7 @@
 #include "Insights/Common/InsightsStyle.h"
 #include "Insights/Filter/ViewModels/Filters.h"
 
-#define LOCTEXT_NAMESPACE "SFilterConfiguratorRow"
+#define LOCTEXT_NAMESPACE "UE::Insights::SFilterConfiguratorRow"
 
 namespace UE
 {
@@ -242,7 +242,7 @@ TSharedRef<SWidget> SFilterConfiguratorRow::GenerateWidgetForColumn(const FName&
 			];
 
 		// Do not show Delete button for the Root node
-		EVisibility DeleteVisibility = FilterConfiguratorNodePtr->GetGroupPtr().IsValid() ? EVisibility::Visible : EVisibility::Hidden;
+		EVisibility DeleteVisibility = FilterConfiguratorNodePtr->GetParent().IsValid() ? EVisibility::Visible : EVisibility::Hidden;
 		RightBox->AddSlot()
 			.VAlign(VAlign_Center)
 			.HAlign(HAlign_Right)
@@ -292,7 +292,7 @@ TSharedRef<SWidget> SFilterConfiguratorRow::AvailableFilters_OnGenerateWidget(TS
 		.AutoWidth()
 		[
 			SNew(STextBlock)
-			.Text(InFilter->Name)
+			.Text(InFilter->GetName())
 			.Margin(2.0f)
 		];
 
@@ -320,12 +320,12 @@ TSharedRef<SWidget> SFilterConfiguratorRow::AvailableFilterOperators_OnGenerateW
 TSharedRef<SWidget> SFilterConfiguratorRow::FilterGroupOperators_OnGenerateWidget(TSharedPtr<FFilterGroupOperator> InFilterGroupOperator)
 {
 	TSharedRef<SHorizontalBox> Widget = SNew(SHorizontalBox);
-	Widget->SetToolTipText(InFilterGroupOperator->Desc);
+	Widget->SetToolTipText(InFilterGroupOperator->GetDesc());
 	Widget->AddSlot()
 		.AutoWidth()
 		[
 			SNew(STextBlock)
-			.Text(InFilterGroupOperator->Name)
+			.Text(InFilterGroupOperator->GetName())
 			.Margin(2.0f)
 		];
 
@@ -336,9 +336,9 @@ END_SLATE_FUNCTION_BUILD_OPTIMIZATION
 
 ////////////////////////////////////////////////////////////////////////////////////////////////////
 
-const TArray<TSharedPtr<struct FFilter>>* SFilterConfiguratorRow::GetAvailableFilters()
+const TArray<TSharedPtr<FFilter>>* SFilterConfiguratorRow::GetAvailableFilters()
 {
-	return &*FilterConfiguratorNodePtr->GetAvailableFilters();
+	return FilterConfiguratorNodePtr->GetAvailableFilters().Get();
 }
 
 ///////////////////////////////////////////////////////////////////////////////////////////////////
@@ -357,7 +357,7 @@ void SFilterConfiguratorRow::AvailableFilters_OnSelectionChanged(TSharedPtr<FFil
 
 FText SFilterConfiguratorRow::AvailableFilters_GetSelectionText() const
 {
-	return FilterConfiguratorNodePtr->GetSelectedFilter()->Name;
+	return FilterConfiguratorNodePtr->GetSelectedFilter()->GetName();
 }
 
 ///////////////////////////////////////////////////////////////////////////////////////////////////
@@ -399,7 +399,7 @@ void SFilterConfiguratorRow::FilterGroupOperators_OnSelectionChanged(TSharedPtr<
 
 FText SFilterConfiguratorRow::FilterGroupOperators_GetSelectionText() const
 {
-	return FilterConfiguratorNodePtr->GetSelectedFilterGroupOperator()->Name;
+	return FilterConfiguratorNodePtr->GetSelectedFilterGroupOperator()->GetName();
 }
 
 ///////////////////////////////////////////////////////////////////////////////////////////////////
@@ -409,7 +409,7 @@ FReply SFilterConfiguratorRow::AddFilter_OnClicked()
 	FFilterConfiguratorNodePtr ChildNode = MakeShared<FFilterConfiguratorNode>(TEXT(""), false);
 	ChildNode->SetAvailableFilters(FilterConfiguratorNodePtr->GetAvailableFilters());
 
-	FilterConfiguratorNodePtr->AddChildAndSetGroupPtr(ChildNode);
+	FilterConfiguratorNodePtr->AddChildAndSetParent(ChildNode);
 	FilterConfiguratorNodePtr->SetExpansion(true);
 	OwnerTablePtr.Pin()->Private_SetItemExpansion(FilterConfiguratorNodePtr, true);
 
@@ -421,12 +421,7 @@ FReply SFilterConfiguratorRow::AddFilter_OnClicked()
 
 FReply SFilterConfiguratorRow::DeleteFilter_OnClicked()
 {
-	FBaseTreeNodeWeak GroupWeakPtr = FilterConfiguratorNodePtr->GetGroupPtr();
-	if (GroupWeakPtr.IsValid())
-	{
-		FFilterConfiguratorNodePtr GroupPtr = StaticCastSharedPtr<FFilterConfiguratorNode>(GroupWeakPtr.Pin());
-		GroupPtr->DeleteChildNode(FilterConfiguratorNodePtr);
-	}
+	FilterConfiguratorNodePtr->RemoveFromParent();
 
 	StaticCastSharedPtr<STreeView<FFilterConfiguratorNodePtr>>(OwnerTablePtr.Pin())->RequestTreeRefresh();
 	return FReply::Handled();
@@ -440,7 +435,7 @@ FReply SFilterConfiguratorRow::AddGroup_OnClicked()
 	ChildNode->SetAvailableFilters(FilterConfiguratorNodePtr->GetAvailableFilters());
 	ChildNode->SetExpansion(true);
 
-	FilterConfiguratorNodePtr->AddChildAndSetGroupPtr(ChildNode);
+	FilterConfiguratorNodePtr->AddChildAndSetParent(ChildNode);
 	FilterConfiguratorNodePtr->SetExpansion(true);
 	OwnerTablePtr.Pin()->Private_SetItemExpansion(FilterConfiguratorNodePtr, true);
 
@@ -452,12 +447,7 @@ FReply SFilterConfiguratorRow::AddGroup_OnClicked()
 
 FReply SFilterConfiguratorRow::DeleteGroup_OnClicked()
 {
-	FBaseTreeNodeWeak GroupWeakPtr = FilterConfiguratorNodePtr->GetGroupPtr();
-	if (GroupWeakPtr.IsValid())
-	{
-		FFilterConfiguratorNodePtr GroupPtr = StaticCastSharedPtr<FFilterConfiguratorNode>(GroupWeakPtr.Pin());
-		GroupPtr->DeleteChildNode(FilterConfiguratorNodePtr);
-	}
+	FilterConfiguratorNodePtr->RemoveFromParent();
 
 	StaticCastSharedPtr<STreeView<FFilterConfiguratorNodePtr>>(OwnerTablePtr.Pin())->RequestTreeRefresh();
 	return FReply::Handled();
@@ -481,7 +471,7 @@ void SFilterConfiguratorRow::OnTextBoxValueCommitted(const FText& InNewText, ETe
 
 FText SFilterConfiguratorRow::GetTextBoxTooltipText() const
 {
-	TSharedPtr<IFilterValueConverter> Converter = FilterConfiguratorNodePtr->GetSelectedFilter()->Converter;
+	const TSharedPtr<IFilterValueConverter>& Converter = FilterConfiguratorNodePtr->GetSelectedFilter()->GetConverter();
 	if (Converter.IsValid())
 	{
 		return Converter->GetTooltipText();
@@ -494,7 +484,7 @@ FText SFilterConfiguratorRow::GetTextBoxTooltipText() const
 
 FText SFilterConfiguratorRow::GetTextBoxHintText() const
 {
-	TSharedPtr<IFilterValueConverter> Converter = FilterConfiguratorNodePtr->GetSelectedFilter()->Converter;
+	const TSharedPtr<IFilterValueConverter>& Converter = FilterConfiguratorNodePtr->GetSelectedFilter()->GetConverter();
 	if (Converter.IsValid())
 	{
 		return Converter->GetHintText();
@@ -507,15 +497,15 @@ FText SFilterConfiguratorRow::GetTextBoxHintText() const
 
 bool SFilterConfiguratorRow::TextBox_OnVerifyTextChanged(const FText& InText, FText& OutErrorMessage)
 {
-	TSharedPtr<IFilterValueConverter> Converter = FilterConfiguratorNodePtr->GetSelectedFilter()->Converter;
+	const TSharedPtr<IFilterValueConverter>& Converter = FilterConfiguratorNodePtr->GetSelectedFilter()->GetConverter();
 	if (Converter.IsValid())
 	{
-		if (FilterConfiguratorNodePtr->GetSelectedFilter()->DataType == EFilterDataType::Int64)
+		if (FilterConfiguratorNodePtr->GetSelectedFilter()->GetDataType() == EFilterDataType::Int64)
 		{
 			int64 Value;
 			return Converter->Convert(InText.ToString(), Value, OutErrorMessage);
 		}		
-		else if (FilterConfiguratorNodePtr->GetSelectedFilter()->DataType == EFilterDataType::Double)
+		else if (FilterConfiguratorNodePtr->GetSelectedFilter()->GetDataType() == EFilterDataType::Double)
 		{
 			double Value;
 			return Converter->Convert(InText.ToString(), Value, OutErrorMessage);
@@ -534,7 +524,7 @@ void SFilterConfiguratorRow::SuggestionTextBox_GetSuggestions(const FString& Tex
 	if (Filter->Is<FFilterWithSuggestions>())
 	{
 		TSharedPtr<FFilterWithSuggestions> FilterWithSuggestions = StaticCastSharedPtr<FFilterWithSuggestions>(Filter);
-		FilterWithSuggestions->Callback(Text, Suggestions);
+		FilterWithSuggestions->GetSuggestions(Text, Suggestions);
 	}
 }
 
diff --git a/Engine/Plugins/Editor/AssetManagerEditor/Source/AssetManagerEditor/Private/Insights/Filter/Widgets/SFilterConfiguratorRow.h b/Engine/Plugins/Editor/AssetManagerEditor/Source/AssetManagerEditor/Private/Insights/Filter/Widgets/SFilterConfiguratorRow.h
--- a/Engine/Plugins/Editor/AssetManagerEditor/Source/AssetManagerEditor/Private/Insights/Filter/Widgets/SFilterConfiguratorRow.h
+++ b/Engine/Plugins/Editor/AssetManagerEditor/Source/AssetManagerEditor/Private/Insights/Filter/Widgets/SFilterConfiguratorRow.h
@@ -15,6 +15,10 @@ namespace UE
 namespace Insights
 {
 
+class FFilter;
+class FFilterGroupOperator;
+class IFilterOperator;
+
 /** Widget that represents a table row in the Filter's tree control. Generates widgets for each column on demand. */
 class SFilterConfiguratorRow : public SMultiColumnTableRow<FFilterConfiguratorNodePtr>
 {
@@ -32,27 +36,27 @@ public:
 	virtual TSharedRef<SWidget> GenerateWidgetForColumn(const FName& InColumnName) override;
 
 private:
-	const TArray<TSharedPtr<struct FFilter>>* GetAvailableFilters();
+	const TArray<TSharedPtr<FFilter>>* GetAvailableFilters();
 
-	TSharedRef<SWidget> AvailableFilters_OnGenerateWidget(TSharedPtr<struct FFilter> InFilter);
+	TSharedRef<SWidget> AvailableFilters_OnGenerateWidget(TSharedPtr<FFilter> InFilter);
 
-	void AvailableFilters_OnSelectionChanged(TSharedPtr<struct FFilter> InFilter, ESelectInfo::Type SelectInfo);
+	void AvailableFilters_OnSelectionChanged(TSharedPtr<FFilter> InFilter, ESelectInfo::Type SelectInfo);
 
 	FText AvailableFilters_GetSelectionText() const;
 
-	const TArray<TSharedPtr<class IFilterOperator>>* GetAvailableFilterOperators();
+	const TArray<TSharedPtr<IFilterOperator>>* GetAvailableFilterOperators();
 
-	TSharedRef<SWidget> AvailableFilterOperators_OnGenerateWidget(TSharedPtr<class IFilterOperator> InFilter);
+	TSharedRef<SWidget> AvailableFilterOperators_OnGenerateWidget(TSharedPtr<IFilterOperator> InFilter);
 
-	void AvailableFilterOperators_OnSelectionChanged(TSharedPtr<class IFilterOperator> InFilter, ESelectInfo::Type SelectInfo);
+	void AvailableFilterOperators_OnSelectionChanged(TSharedPtr<IFilterOperator> InFilter, ESelectInfo::Type SelectInfo);
 
 	FText AvailableFilterOperators_GetSelectionText() const;
 
-	const TArray<TSharedPtr<struct FFilterGroupOperator>>* GetFilterGroupOperators();
+	const TArray<TSharedPtr<FFilterGroupOperator>>* GetFilterGroupOperators();
 
-	TSharedRef<SWidget> FilterGroupOperators_OnGenerateWidget(TSharedPtr<struct FFilterGroupOperator> InFilter);
+	TSharedRef<SWidget> FilterGroupOperators_OnGenerateWidget(TSharedPtr<FFilterGroupOperator> InFilter);
 
-	void FilterGroupOperators_OnSelectionChanged(TSharedPtr<struct FFilterGroupOperator> InFilter, ESelectInfo::Type SelectInfo);
+	void FilterGroupOperators_OnSelectionChanged(TSharedPtr<FFilterGroupOperator> InFilter, ESelectInfo::Type SelectInfo);
 
 	FText FilterGroupOperators_GetSelectionText() const;
 
@@ -79,11 +83,11 @@ private:
 private:
 	FFilterConfiguratorNodePtr FilterConfiguratorNodePtr;
 
-	TSharedPtr<SComboBox<TSharedPtr<struct FFilter>>> FilterTypeComboBox;
+	TSharedPtr<SComboBox<TSharedPtr<FFilter>>> FilterTypeComboBox;
 
-	TSharedPtr<SComboBox<TSharedPtr<class IFilterOperator>>> FilterOperatorComboBox;
+	TSharedPtr<SComboBox<TSharedPtr<IFilterOperator>>> FilterOperatorComboBox;
 
-	TSharedPtr<SComboBox<TSharedPtr<struct FFilterGroupOperator>>> FilterGroupOperatorComboBox;
+	TSharedPtr<SComboBox<TSharedPtr<FFilterGroupOperator>>> FilterGroupOperatorComboBox;
 
 	FString SuggestionTextBoxValue;
 };
diff --git a/Engine/Plugins/Editor/AssetManagerEditor/Source/AssetManagerEditor/Private/Insights/Table/ViewModels/BaseTreeNode.cpp b/Engine/Plugins/Editor/AssetManagerEditor/Source/AssetManagerEditor/Private/Insights/Table/ViewModels/BaseTreeNode.cpp
--- a/Engine/Plugins/Editor/AssetManagerEditor/Source/AssetManagerEditor/Private/Insights/Table/ViewModels/BaseTreeNode.cpp
+++ b/Engine/Plugins/Editor/AssetManagerEditor/Source/AssetManagerEditor/Private/Insights/Table/ViewModels/BaseTreeNode.cpp
@@ -33,8 +33,8 @@ const FText FBaseTreeNode::GetExtraDisplayName() const
 {
 	if (IsGroup())
 	{
-		const int32 NumChildren = GroupData->Children.Num();
-		const int32 NumFilteredChildren = GroupData->FilteredChildren.Num();
+		const int32 NumChildren = GetChildrenCount();
+		const int32 NumFilteredChildren = GetFilteredChildrenCount();
 
 		if (NumFilteredChildren == NumChildren)
 		{
@@ -86,16 +86,16 @@ FLinearColor FBaseTreeNode::GetDefaultColor(bool bIsGroupNode)
 
 ////////////////////////////////////////////////////////////////////////////////////////////////////
 
-void FBaseTreeNode::SortChildrenAscending(const ITableCellValueSorter& Sorter)
+void FBaseTreeNode::SortChildren(const ITableCellValueSorter& Sorter, ESortMode SortMode)
 {
-	Sorter.Sort(GroupData->Children, ESortMode::Ascending);
+	Sorter.Sort(GroupData->Children, SortMode);
 }
 
 ////////////////////////////////////////////////////////////////////////////////////////////////////
 
-void FBaseTreeNode::SortChildrenDescending(const ITableCellValueSorter& Sorter)
+void FBaseTreeNode::SortFilteredChildren(const ITableCellValueSorter& Sorter, ESortMode SortMode)
 {
-	Sorter.Sort(GroupData->Children, ESortMode::Descending);
+	Sorter.Sort(*GroupData->FilteredChildrenPtr, SortMode);
 }
 
 ////////////////////////////////////////////////////////////////////////////////////////////////////
diff --git a/Engine/Plugins/Editor/AssetManagerEditor/Source/AssetManagerEditor/Private/Insights/Table/ViewModels/BaseTreeNode.h b/Engine/Plugins/Editor/AssetManagerEditor/Source/AssetManagerEditor/Private/Insights/Table/ViewModels/BaseTreeNode.h
--- a/Engine/Plugins/Editor/AssetManagerEditor/Source/AssetManagerEditor/Private/Insights/Table/ViewModels/BaseTreeNode.h
+++ b/Engine/Plugins/Editor/AssetManagerEditor/Source/AssetManagerEditor/Private/Insights/Table/ViewModels/BaseTreeNode.h
@@ -16,6 +16,7 @@ namespace Insights
 {
 
 class ITableCellValueSorter;
+enum class ESortMode;
 
 ////////////////////////////////////////////////////////////////////////////////////////////////////
 
@@ -40,22 +41,25 @@ class FBaseTreeNode : public TSharedFromThis<FBaseTreeNode>
 	INSIGHTS_DECLARE_RTTI_BASE(FBaseTreeNode)
 
 protected:
+	enum class EGroupNodeFlags : uint32
+	{
+		None = 0,
+		IsExpanded = 1 << 0,
+	};
+
 	struct FGroupNodeData
 	{
 		/** Children of the group node. */
 		TArray<FBaseTreeNodePtr> Children;
 
-		/** Filtered children of the group node. */
-		TArray<FBaseTreeNodePtr> FilteredChildren;
+		/**
+		 * A pointer to the filtered children of the group node.
+		 * If filter is empty (no filtering), this points directly to Children array.
+		 */
+		TArray<FBaseTreeNodePtr>* FilteredChildrenPtr = &Children;
 
-		/** Descriptive text for the group node to display in a tooltip. */
-		FText Tooltip;
-
-		/** A generic pointer to a data context for the group node. Can be nullptr. */
-		void* Context = nullptr;
-
-		/** Whether the group node should be expanded or not. */
-		bool bIsExpanded = false;
+		/** Various flags specific to a group node. */
+		EGroupNodeFlags Flags = EGroupNodeFlags::None;
 	};
 
 public:
@@ -63,26 +67,14 @@ public:
 	explicit FBaseTreeNode(const FName InName, bool bInIsGroup)
 		: DefaultSortOrder(0)
 		, Name(InName)
+		, Parent(nullptr)
 		, GroupData(bInIsGroup ? new FGroupNodeData() : &DefaultGroupData)
 	{
 	}
 
 	virtual ~FBaseTreeNode()
 	{
-		if (GroupData != &DefaultGroupData)
-		{
-			delete GroupData;
-			GroupData = nullptr;
-		}
-	}
-
-	uint32 GetDefaultSortOrder() const
-	{
-		return DefaultSortOrder;
-	}
-	void SetDefaultSortOrder(uint32 Order)
-	{
-		DefaultSortOrder = Order;
+		RemoveGroupData();
 	}
 
 	/**
@@ -111,20 +103,9 @@ public:
 	/**
 	 * @returns a descriptive text for this node to display in a tooltip.
 	 */
-	virtual const FText GetTooltip() const
+	virtual const FText GetTooltipText() const
 	{
-		return GroupData->Tooltip;
-	}
-
-	/**
-	 * Sets a descriptive text for this node to display in a tooltip.
-	 */
-	virtual void SetTooltip(const FText& InTooltip)
-	{
-		if (IsGroup())
-		{
-			GroupData->Tooltip = InTooltip;
-		}
+		return FText::GetEmpty();
 	}
 
 	/**
@@ -153,25 +134,28 @@ public:
 		return GetDefaultColor(IsGroup());
 	}
 
+	//////////////////////////////////////////////////
+	// Parent Node
+
 	/**
-	 * @returns a pointer to a data context for this node. Only available for group nodes.
+	 * @returns a weak reference to the parent group of this node; may be invalid.
 	 */
-	virtual void* GetContext() const
+	FBaseTreeNodeWeak GetParentWeak() const
 	{
-		return GroupData->Context;
+		return Parent;
 	}
 
 	/**
-	 * Sets a pointer to a data context for this node. Only available for group nodes.
+	 * @returns a reference to the parent group of this node; may be invalid.
 	 */
-	virtual void SetContext(void* InContext)
+	FBaseTreeNodePtr GetParent() const
 	{
-		if (IsGroup())
-		{
-			GroupData->Context = InContext;
-		}
+		return Parent.Pin();
 	}
 
+	//////////////////////////////////////////////////
+	// Group Node
+
 	/**
 	 * @returns true if this node is a group node.
 	 */
@@ -191,20 +175,29 @@ public:
 		}
 	}
 
-	/**
-	 * @returns the number of children nodes.
-	 */
-	int32 GetChildrenCount() const
+	void RemoveGroupData()
 	{
-		return GroupData->Children.Num();
+		if (GroupData != &DefaultGroupData)
+		{
+			if (GroupData->FilteredChildrenPtr != &GroupData->Children)
+			{
+				delete GroupData->FilteredChildrenPtr;
+			}
+
+			delete GroupData;
+			GroupData = &DefaultGroupData;
+		}
 	}
 
+	//////////////////////////////////////////////////
+	// Children Nodes
+
 	/**
-	 * @returns the number of filtered children nodes.
+	 * @returns the number of children nodes.
 	 */
-	int32 GetFilteredChildrenCount() const
+	int32 GetChildrenCount() const
 	{
-		return GroupData->FilteredChildren.Num();
+		return GroupData->Children.Num();
 	}
 
 	/**
@@ -221,20 +214,6 @@ public:
 		}
 	}
 
-	/**
-	 * Enumerates the filtered children nodes.
-	 */
-	void EnumerateFilteredChildren(TFunction<bool(const FBaseTreeNodePtr&)> Callback) const
-	{
-		for (const FBaseTreeNodePtr& ChildNode : GroupData->FilteredChildren)
-		{
-			if (!Callback(ChildNode))
-			{
-				break;
-			}
-		}
-	}
-
 	/**
 	 * @returns a const reference to the children nodes of this group.
 	 */
@@ -243,75 +222,124 @@ public:
 		return GroupData->Children;
 	}
 
-	/**
-	 * @returns a const reference to the children nodes that should be visible to the UI based on filtering.
-	 */
-	const TArray<FBaseTreeNodePtr>& GetFilteredChildren() const
+	void SortChildren(const ITableCellValueSorter& Sorter, ESortMode SortMode);
+
+	template <typename PredicateType>
+	void SortChildren(PredicateType Predicate)
 	{
-		return GroupData->FilteredChildren;
+		GroupData->Children.Sort(Predicate);
 	}
 
-	/**
-	 * @returns a weak reference to the group of this node, may be invalid.
-	 */
-	FBaseTreeNodeWeak GetGroupPtr() const
+	/** Adds specified node to the children nodes. */
+	void AddChild(const FBaseTreeNodePtr& ChildPtr)
 	{
-		return GroupPtr;
+		GroupData->Children.Add(ChildPtr);
 	}
 
-	/**
-	 * @returns a weak reference to the group of this node, may be invalid.
-	 */
-	FBaseTreeNodePtr GetParentNode() const
+	/** Adds specified node to the children nodes. Also sets the current node as the parent group of the specified node. */
+	void AddChildAndSetParent(const FBaseTreeNodePtr& ChildPtr)
 	{
-		return GroupPtr.Pin();
+		ChildPtr->Parent = AsWeak();
+		GroupData->Children.Add(ChildPtr);
+	}
+
+	/** Removes this node from the parent's children nodes (preserving the order). Also resets the parent group. */
+	void RemoveFromParent()
+	{
+		FBaseTreeNodePtr ParentNodePtr = Parent.Pin();
+		if (ParentNodePtr.IsValid())
+		{
+			ParentNodePtr->GroupData->Children.RemoveSingle(AsShared());
+		}
+		Parent.Reset();
 	}
 
-	void SortChildrenAscending(const ITableCellValueSorter& Sorter);
-	void SortChildrenDescending(const ITableCellValueSorter& Sorter);
+	/** Removes this node from the parent's children nodes (not preserving the order). Also resets the parent group. */
+	void RemoveFromParentSwap()
+	{
+		FBaseTreeNodePtr ParentNodePtr = Parent.Pin();
+		if (ParentNodePtr.IsValid())
+		{
+			ParentNodePtr->GroupData->Children.RemoveSwap(AsShared());
+		}
+		Parent.Reset();
+	}
 
-	template <class PREDICATE_CLASS>
-	void SortChildren(const PREDICATE_CLASS& Predicate)
+	/** Removes the specified node from the children nodes (preserving the order). If successful, resets the parent group of the specified node. */
+	int32 RemoveChild(const FBaseTreeNodePtr& ChildPtr)
 	{
-		GroupData->Children.Sort(Predicate);
+		if (GroupData->Children.RemoveSingle(ChildPtr))
+		{
+			ChildPtr->Parent.Reset();
+			return 1;
+		}
+		return 0;
 	}
 
-	/** Adds specified node to the children nodes. Also sets the current node as the group of the specified node. */
-	void AddChildAndSetGroupPtr(const FBaseTreeNodePtr& ChildPtr)
+	/** Removes the specified node from the children nodes (not preserving the order). If successful, resets the parent group of the specified node. */
+	int32 RemoveChildSwap(const FBaseTreeNodePtr& ChildPtr)
 	{
-		ChildPtr->GroupPtr = AsWeak();
-		GroupData->Children.Add(ChildPtr);
+		if (GroupData->Children.RemoveSwap(ChildPtr))
+		{
+			ChildPtr->Parent.Reset();
+			return 1;
+		}
+		return 0;
 	}
 
-	void SetGroupPtrForAllChildren()
+	void SetParent(const FBaseTreeNodeWeak& InParentWeak)
+	{
+		Parent = InParentWeak;
+	}
+
+	void SetParentForChildren()
+	{
+		FBaseTreeNodeWeak ThisNode = AsWeak();
+		for (FBaseTreeNodePtr& NodePtr : GroupData->Children)
+		{
+			NodePtr->Parent = ThisNode;
+		}
+	}
+
+	void SetParentForChildrenRec()
 	{
 		FBaseTreeNodeWeak ThisNode = AsWeak();
 		for (FBaseTreeNodePtr& NodePtr : GroupData->Children)
 		{
-			NodePtr->GroupPtr = ThisNode;
+			NodePtr->Parent = ThisNode;
+			NodePtr->SetParentForChildrenRec();
 		}
 	}
 
-	void ResetGroupPtrForAllChildren()
+	void ResetParentForChildren()
 	{
 		for (FBaseTreeNodePtr& NodePtr : GroupData->Children)
 		{
-			NodePtr->GroupPtr = nullptr;
+			NodePtr->Parent.Reset();
+		}
+	}
+
+	void ResetParentForChildrenRec()
+	{
+		for (FBaseTreeNodePtr& NodePtr : GroupData->Children)
+		{
+			NodePtr->ResetParentForChildrenRec();
+			NodePtr->Parent.Reset();
 		}
 	}
 
 	/** Clears children. */
 	void ClearChildren(int32 NewSize = 0)
 	{
-		ResetGroupPtrForAllChildren();
+		ResetParentForChildren();
 		GroupData->Children.Reset(NewSize);
 	}
 
 	void SwapChildren(TArray<FBaseTreeNodePtr>& NewChildren)
 	{
-		ResetGroupPtrForAllChildren();
+		ResetParentForChildren();
 		Swap(NewChildren, GroupData->Children);
-		SetGroupPtrForAllChildren();
+		SetParentForChildren();
 	}
 
 	void SwapChildrenFast(TArray<FBaseTreeNodePtr>& NewChildren)
@@ -319,16 +347,86 @@ public:
 		Swap(NewChildren, GroupData->Children);
 	}
 
-	/** Adds specified child to the filtered children nodes. */
-	void AddFilteredChild(const FBaseTreeNodePtr& ChildPtr)
+	//////////////////////////////////////////////////
+	// Filtered Children Nodes
+
+	virtual bool IsFiltered() const
 	{
-		GroupData->FilteredChildren.Add(ChildPtr);
+		return false;
 	}
 
-	/** Clears the filtered children nodes. */
+	/**
+	 * @returns the number of filtered children nodes.
+	 */
+	int32 GetFilteredChildrenCount() const
+	{
+		return GroupData->FilteredChildrenPtr->Num();
+	}
+
+	/**
+	 * @returns a const reference to the filtered child node at specified index.
+	 */
+	const FBaseTreeNodePtr& GetFilteredChildNode(int32 Index) const
+	{
+		return (*GroupData->FilteredChildrenPtr)[Index];
+	}
+
+	/**
+	 * @returns a reference to the filtered child node at specified index.
+	 */
+	FBaseTreeNodePtr& GetFilteredChildNode(int32 Index)
+	{
+		return (*GroupData->FilteredChildrenPtr)[Index];
+	}
+
+	/**
+	 * Enumerates the filtered children nodes.
+	 */
+	void EnumerateFilteredChildren(TFunction<bool(const FBaseTreeNodePtr&)> Callback) const
+	{
+		for (const FBaseTreeNodePtr& ChildNode : *GroupData->FilteredChildrenPtr)
+		{
+			if (!Callback(ChildNode))
+			{
+				break;
+			}
+		}
+	}
+
+	/**
+	 * @returns a const reference to the children nodes that should be visible to the UI based on filtering.
+	 */
+	const TArray<FBaseTreeNodePtr>& GetFilteredChildren() const
+	{
+		return *GroupData->FilteredChildrenPtr;
+	}
+
+	void SortFilteredChildren(const ITableCellValueSorter& Sorter, ESortMode SortMode);
+
+	template <typename PredicateType>
+	void SortFilteredChildren(PredicateType Predicate)
+	{
+		GroupData->FilteredChildrenPtr->Sort(Predicate);
+	}
+
+	/** Clears (removes all) the filtered children nodes. */
 	void ClearFilteredChildren(int32 NewSize = 0)
 	{
-		GroupData->FilteredChildren.Reset(NewSize);
+		if (GroupData->FilteredChildrenPtr == &GroupData->Children)
+		{
+			GroupData->FilteredChildrenPtr = new TArray<FBaseTreeNodePtr>();
+		}
+		GroupData->FilteredChildrenPtr->Reset(NewSize);
+	}
+
+	/** Adds specified child to the filtered children nodes. */
+	void AddFilteredChild(const FBaseTreeNodePtr& ChildPtr)
+	{
+		if (GroupData->FilteredChildrenPtr == &GroupData->Children)
+		{
+			GroupData->FilteredChildrenPtr = new TArray<FBaseTreeNodePtr>();
+		}
+		GroupData->FilteredChildrenPtr->Add(ChildPtr);
 	}
 
 	/**
@@ -337,9 +435,11 @@ public:
 	 */
 	void ResetFilteredChildren()
 	{
-		const int32 Size = GroupData->Children.Num();
-		GroupData->FilteredChildren.Reset();
-		GroupData->FilteredChildren.Append(GroupData->Children);
+		if (GroupData->FilteredChildrenPtr != &GroupData->Children)
+		{
+			delete GroupData->FilteredChildrenPtr;
+			GroupData->FilteredChildrenPtr = &GroupData->Children;
+		}
 	}
 
 	/**
@@ -355,19 +455,38 @@ public:
 		}
 	}
 
+	//////////////////////////////////////////////////
+
 	bool IsExpanded() const
 	{
-		return GroupData->bIsExpanded;
+		return ((uint32)GroupData->Flags & (uint32)EGroupNodeFlags::IsExpanded) != 0;
 	}
 
 	void SetExpansion(bool bOnOff)
 	{
 		if (IsGroup())
 		{
-			GroupData->bIsExpanded = bOnOff;
+			if (bOnOff)
+			{
+				GroupData->Flags = (EGroupNodeFlags)((uint32)GroupData->Flags | (uint32)EGroupNodeFlags::IsExpanded);
+			}
+			else
+			{
+				GroupData->Flags = (EGroupNodeFlags)((uint32)GroupData->Flags & ~(uint32)EGroupNodeFlags::IsExpanded);
+			}
 		}
 	}
 
+	uint32 GetDefaultSortOrder() const
+	{
+		return DefaultSortOrder;
+	}
+
+	void SetDefaultSortOrder(uint32 Order)
+	{
+		DefaultSortOrder = Order;
+	}
+
 protected:
 	/**
 	 * @returns a reference to the children nodes of this group.
@@ -377,11 +496,6 @@ protected:
 		return GroupData->Children;
 	}
 
-	void SetGroupPtr(FBaseTreeNodePtr InGroupPtr)
-	{
-		GroupPtr = InGroupPtr;
-	}
-
 private:
 	/** The default sort order. Index used to optimize sorting. */
 	int32 DefaultSortOrder;
@@ -389,8 +503,8 @@ private:
 	/** The name of this node. */
 	const FName Name;
 
-	/** A weak pointer to the group/parent of this node. */
-	FBaseTreeNodeWeak GroupPtr;
+	/** A weak pointer to the parent group of this node. */
+	FBaseTreeNodeWeak Parent;
 
 	/** The struct containing properties of a group node. It is allocated only for group nodes. */
 	FGroupNodeData* GroupData;
diff --git a/Engine/Plugins/Editor/AssetManagerEditor/Source/AssetManagerEditor/Private/Insights/Table/ViewModels/TreeNodeGrouping.cpp b/Engine/Plugins/Editor/AssetManagerEditor/Source/AssetManagerEditor/Private/Insights/Table/ViewModels/TreeNodeGrouping.cpp
--- a/Engine/Plugins/Editor/AssetManagerEditor/Source/AssetManagerEditor/Private/Insights/Table/ViewModels/TreeNodeGrouping.cpp
+++ b/Engine/Plugins/Editor/AssetManagerEditor/Source/AssetManagerEditor/Private/Insights/Table/ViewModels/TreeNodeGrouping.cpp
@@ -50,7 +50,7 @@ void FTreeNodeGrouping::GroupNodes(const TArray<FTableTreeNodePtr>& Nodes, FTabl
 
 		if (NodePtr->IsGroup())
 		{
-			ParentGroup.AddChildAndSetGroupPtr(NodePtr);
+			ParentGroup.AddChildAndSetParent(NodePtr);
 			continue;
 		}
 
@@ -62,7 +62,7 @@ void FTreeNodeGrouping::GroupNodes(const TArray<FTableTreeNodePtr>& Nodes, FTabl
 		{
 			GroupPtr = MakeShared<FTableTreeNode>(GroupInfo.Name, InParentTable);
 			GroupPtr->SetExpansion(GroupInfo.IsExpanded);
-			ParentGroup.AddChildAndSetGroupPtr(GroupPtr);
+			ParentGroup.AddChildAndSetParent(GroupPtr);
 			GroupMap.Add(GroupInfo.Name, GroupPtr);
 		}
 		else
@@ -70,7 +70,7 @@ void FTreeNodeGrouping::GroupNodes(const TArray<FTableTreeNodePtr>& Nodes, FTabl
 			GroupPtr = *GroupPtrPtr;
 		}
 
-		GroupPtr->AddChildAndSetGroupPtr(NodePtr);
+		GroupPtr->AddChildAndSetParent(NodePtr);
 	}
 }
 
@@ -96,7 +96,7 @@ void FTreeNodeGroupingFlat::GroupNodes(const TArray<FTableTreeNodePtr>& Nodes, F
 
 	FTableTreeNodePtr GroupPtr = MakeShared<FTableTreeNode>(FName(TEXT("All")), InParentTable);
 	GroupPtr->SetExpansion(true);
-	ParentGroup.AddChildAndSetGroupPtr(GroupPtr);
+	ParentGroup.AddChildAndSetParent(GroupPtr);
 
 	GroupPtr->ClearChildren(Nodes.Num());
 	for (FTableTreeNodePtr NodePtr : Nodes)
@@ -106,7 +106,7 @@ void FTreeNodeGroupingFlat::GroupNodes(const TArray<FTableTreeNodePtr>& Nodes, F
 			return;
 		}
 
-		GroupPtr->AddChildAndSetGroupPtr(NodePtr);
+		GroupPtr->AddChildAndSetParent(NodePtr);
 	}
 }
 
@@ -182,7 +182,7 @@ void FTreeNodeGroupingByUniqueValueCString::GroupNodes(const TArray<FTableTreeNo
 
 		if (NodePtr->IsGroup())
 		{
-			ParentGroup.AddChildAndSetGroupPtr(NodePtr);
+			ParentGroup.AddChildAndSetParent(NodePtr);
 			continue;
 		}
 
@@ -204,7 +204,7 @@ void FTreeNodeGroupingByUniqueValueCString::GroupNodes(const TArray<FTableTreeNo
 					{
 						GroupPtr = MakeShared<FTableTreeNode>(GroupName, InParentTable);
 						GroupPtr->SetExpansion(false);
-						ParentGroup.AddChildAndSetGroupPtr(GroupPtr);
+						ParentGroup.AddChildAndSetParent(GroupPtr);
 						GroupNameMap.Add(GroupName, GroupPtr);
 					}
 					else
@@ -224,7 +224,7 @@ void FTreeNodeGroupingByUniqueValueCString::GroupNodes(const TArray<FTableTreeNo
 				{
 					EmptyGroupPtr = MakeShared<FTableTreeNode>(FName(TEXT("N/A")), InParentTable);
 					EmptyGroupPtr->SetExpansion(false);
-					ParentGroup.AddChildAndSetGroupPtr(EmptyGroupPtr);
+					ParentGroup.AddChildAndSetParent(EmptyGroupPtr);
 				}
 				GroupPtr = EmptyGroupPtr;
 			}
@@ -235,12 +235,12 @@ void FTreeNodeGroupingByUniqueValueCString::GroupNodes(const TArray<FTableTreeNo
 			{
 				UnsetGroupPtr = MakeShared<FTableTreeNode>(FName(TEXT("<unset>")), InParentTable);
 				UnsetGroupPtr->SetExpansion(false);
-				ParentGroup.AddChildAndSetGroupPtr(UnsetGroupPtr);
+				ParentGroup.AddChildAndSetParent(UnsetGroupPtr);
 			}
 			GroupPtr = UnsetGroupPtr;
 		}
 
-		GroupPtr->AddChildAndSetGroupPtr(NodePtr);
+		GroupPtr->AddChildAndSetParent(NodePtr);
 	}
 }
 
@@ -345,7 +345,7 @@ void FTreeNodeGroupingByPathBreakdown::GroupNodes(
 
 		if (NodePtr->IsGroup())
 		{
-			ParentGroup.AddChildAndSetGroupPtr(NodePtr);
+			ParentGroup.AddChildAndSetParent(NodePtr);
 			continue;
 		}
 
@@ -384,7 +384,7 @@ void FTreeNodeGroupingByPathBreakdown::GroupNodes(
 						{
 							FTableTreeNodePtr GroupPtr = MakeShared<FTableTreeNode>(DirName, InParentTable);
 							GroupPtr->SetExpansion(false);
-							Directory->Node->AddChildAndSetGroupPtr(GroupPtr);
+							Directory->Node->AddChildAndSetParent(GroupPtr);
 							Directory = new FDirectory(DirName, Directory, GroupPtr.Get());
 							Directory->Parent->ChildrenMap.Add(DirName, Directory);
 							Directories.Add(Directory);
@@ -405,7 +405,7 @@ void FTreeNodeGroupingByPathBreakdown::GroupNodes(
 
 		if (Directory != Root)
 		{
-			Directory->Node->AddChildAndSetGroupPtr(NodePtr);
+			Directory->Node->AddChildAndSetParent(NodePtr);
 		}
 		else
 		{
@@ -414,10 +414,10 @@ void FTreeNodeGroupingByPathBreakdown::GroupNodes(
 				static FName NotAvailableName(TEXT("N/A"));
 				FTableTreeNodePtr NotAvailableNode = MakeShared<FTableTreeNode>(NotAvailableName, InParentTable);
 				NotAvailableNode->SetExpansion(false);
-				ParentGroup.AddChildAndSetGroupPtr(NotAvailableNode);
+				ParentGroup.AddChildAndSetParent(NotAvailableNode);
 				UnsetGroupPtr = NotAvailableNode.Get();
 			}
-			UnsetGroupPtr->AddChildAndSetGroupPtr(NodePtr);
+			UnsetGroupPtr->AddChildAndSetParent(NodePtr);
 		}
 	}
 
diff --git a/Engine/Plugins/Editor/AssetManagerEditor/Source/AssetManagerEditor/Private/Insights/Table/ViewModels/TreeNodeGrouping.h b/Engine/Plugins/Editor/AssetManagerEditor/Source/AssetManagerEditor/Private/Insights/Table/ViewModels/TreeNodeGrouping.h
--- a/Engine/Plugins/Editor/AssetManagerEditor/Source/AssetManagerEditor/Private/Insights/Table/ViewModels/TreeNodeGrouping.h
+++ b/Engine/Plugins/Editor/AssetManagerEditor/Source/AssetManagerEditor/Private/Insights/Table/ViewModels/TreeNodeGrouping.h
@@ -165,7 +165,7 @@ void TTreeNodeGroupingByUniqueValue<Type>::GroupNodes(const TArray<FTableTreeNod
 
 		if (NodePtr->IsGroup())
 		{
-			ParentGroup.AddChildAndSetGroupPtr(NodePtr);
+			ParentGroup.AddChildAndSetParent(NodePtr);
 			continue;
 		}
 
@@ -183,7 +183,7 @@ void TTreeNodeGroupingByUniqueValue<Type>::GroupNodes(const TArray<FTableTreeNod
 				const FName GroupName = GetGroupName(Column, *NodePtr);
 				GroupPtr = MakeShared<FTableTreeNode>(GroupName, InParentTable);
 				GroupPtr->SetExpansion(false);
-				ParentGroup.AddChildAndSetGroupPtr(GroupPtr);
+				ParentGroup.AddChildAndSetParent(GroupPtr);
 				GroupMap.Add(Value, GroupPtr);
 			}
 			else
@@ -197,12 +197,12 @@ void TTreeNodeGroupingByUniqueValue<Type>::GroupNodes(const TArray<FTableTreeNod
 			{
 				UnsetGroupPtr = MakeShared<FTableTreeNode>(FName(TEXT("<unset>")), InParentTable);
 				UnsetGroupPtr->SetExpansion(false);
-				ParentGroup.AddChildAndSetGroupPtr(UnsetGroupPtr);
+				ParentGroup.AddChildAndSetParent(UnsetGroupPtr);
 			}
 			GroupPtr = UnsetGroupPtr;
 		}
 
-		GroupPtr->AddChildAndSetGroupPtr(NodePtr);
+		GroupPtr->AddChildAndSetParent(NodePtr);
 	}
 }
 
diff --git a/Engine/Plugins/Editor/AssetManagerEditor/Source/AssetManagerEditor/Private/Insights/Table/Widgets/STableTreeView.cpp b/Engine/Plugins/Editor/AssetManagerEditor/Source/AssetManagerEditor/Private/Insights/Table/Widgets/STableTreeView.cpp
--- a/Engine/Plugins/Editor/AssetManagerEditor/Source/AssetManagerEditor/Private/Insights/Table/Widgets/STableTreeView.cpp
+++ b/Engine/Plugins/Editor/AssetManagerEditor/Source/AssetManagerEditor/Private/Insights/Table/Widgets/STableTreeView.cpp
@@ -9,15 +9,10 @@
 #include "Framework/MultiBox/MultiBoxBuilder.h"
 #include "HAL/PlatformApplicationMisc.h"
 #include "HAL/PlatformFileManager.h"
-#include "Layout/WidgetPath.h"
 #include "Logging/MessageLog.h"
 #include "SlateOptMacros.h"
-#include "Widgets/Input/SCheckBox.h"
-#include "Widgets/Layout/SScrollBox.h"
 #include "Widgets/Input/SSearchBox.h"
 #include "Widgets/Images/SImage.h"
-#include "Widgets/Layout/SGridPanel.h"
-#include "Widgets/Layout/SSeparator.h"
 #include "Widgets/SToolTip.h"
 #include "Widgets/Views/STableViewBase.h"
 
@@ -142,6 +137,7 @@
 
 STableTreeView::STableTreeView()
 	: Root(MakeShared<FTableTreeNode>(RootNodeName, Table))
+	, FilteredNodesPtr(&TableRowNodes)
 	, ColumnBeingSorted(GetDefaultColumnBeingSorted())
 	, ColumnSortMode(GetDefaultColumnSortMode())
 {
@@ -151,7 +147,7 @@
 
 STableTreeView::~STableTreeView()
 {
-	checkf(bIsCloseScheduled, TEXT("TableTreeView running in async mode was closed but OnClose() was not called. This can lead to a crash. Call OnClose() from the owner tab/window."))
+	ensureMsgf(bIsCloseScheduled, TEXT("TableTreeView running in async mode was closed but OnClose() was not called. This can lead to a crash. Call OnClose() from the owner tab/window."));
 
 	if (CurrentAsyncOpFilterConfigurator)
 	{
@@ -303,7 +299,7 @@
 					[
 						SNew(STextBlock)
 						.Text(this, &STableTreeView::GetTreeViewBannerText)
-						.ColorAndOpacity(FLinearColor(1.0f, 1.0f, 1.0f, 1.0f))
+						.ColorAndOpacity(FSlateColor(EStyleColor::Foreground))
 					]
 				]
 			]
@@ -328,14 +324,9 @@
 		WidgetContent
 	];
 
-	// Create the search filters: text based, type based etc.
-	TextFilter = MakeShared<FTableTreeNodeTextFilter>(FTableTreeNodeTextFilter::FItemToStringArray::CreateStatic(&STableTreeView::HandleItemToStringArray));
-	Filters = MakeShared<FTableTreeNodeFilterCollection>();
-	CurrentAsyncOpTextFilter = MakeShared<FTableTreeNodeTextFilter>(FTableTreeNodeTextFilter::FItemToStringArray::CreateStatic(&STableTreeView::HandleItemToStringArray));
-	Filters->Add(CurrentAsyncOpTextFilter);
-
+	InitNodeFiltering();
+	InitHierarchyFiltering();
 	InitializeAndShowHeaderColumns();
-
 	InitCommandList();
 }
 END_SLATE_FUNCTION_BUILD_OPTIMIZATION
@@ -343,38 +334,5 @@
 ////////////////////////////////////////////////////////////////////////////////////////////////////
 
 BEGIN_SLATE_FUNCTION_BUILD_OPTIMIZATION
-TSharedRef<SWidget> STableTreeView::ConstructSearchBox()
-{
-	SAssignNew(SearchBox, SSearchBox)
-		.HintText(LOCTEXT("SearchBox_Hint", "Search"))
-		.OnTextChanged(this, &STableTreeView::SearchBox_OnTextChanged)
-		.IsEnabled(this, &STableTreeView::SearchBox_IsEnabled)
-		.ToolTipText(this, &STableTreeView::SearchBox_GetTooltipText);
-
-	return SearchBox.ToSharedRef();
-}
-END_SLATE_FUNCTION_BUILD_OPTIMIZATION
-
-////////////////////////////////////////////////////////////////////////////////////////////////////
-
-BEGIN_SLATE_FUNCTION_BUILD_OPTIMIZATION
-TSharedRef<SWidget> STableTreeView::ConstructAdvancedFiltersButton()
-{
-	return SNew(SButton)
-		.ContentPadding(FMargin(-2.0f, 2.0f, -2.0f, 2.0f))
-		//.Text(LOCTEXT("AdvancedFiltersBtn_Text", "Advanced Filters"))
-		.ToolTipText(this, &STableTreeView::AdvancedFilters_GetTooltipText)
-		.OnClicked(this, &STableTreeView::OnAdvancedFiltersClicked)
-		.IsEnabled(this, &STableTreeView::AdvancedFilters_ShouldBeEnabled)
-		[
-			SNew(SImage)
-			.Image(FInsightsStyle::GetBrush("Icons.ClassicFilterConfig"))
-		];
-}
-END_SLATE_FUNCTION_BUILD_OPTIMIZATION
-
-////////////////////////////////////////////////////////////////////////////////////////////////////
-
-BEGIN_SLATE_FUNCTION_BUILD_OPTIMIZATION
 TSharedRef<SWidget> STableTreeView::ConstructHierarchyBreadcrumbTrail()
 {
@@ -428,7 +386,7 @@
 		.Padding(4.0f, 0.0f, 0.0f, 0.0f)
 		.VAlign(VAlign_Center)
 		[
-			ConstructAdvancedFiltersButton()
+			ConstructFilterConfiguratorButton()
 		]
 	];
 
@@ -493,6 +451,8 @@
 
 TSharedPtr<SWidget> STableTreeView::TreeView_GetMenuContent()
 {
+	FSlateApplication::Get().CloseToolTip();
+
 	TArray<FTableTreeNodePtr> SelectedNodes;
 	const int32 NumSelectedNodes = TreeView->GetSelectedItems(SelectedNodes);
 	FTableTreeNodePtr SelectedNode = NumSelectedNodes ? SelectedNodes[0] : nullptr;
@@ -953,43 +913,34 @@
 			}
 		}
 	}
-
-	if (LazyInvalidateWidgetReason != EInvalidateWidgetReason::None)
-	{
-		TreeView->Invalidate(LazyInvalidateWidgetReason);
-		LazyInvalidateWidgetReason = EInvalidateWidgetReason::None;
-	}
 }
 
 ////////////////////////////////////////////////////////////////////////////////////////////////////
 
 void STableTreeView::UpdateTree()
 {
+	OnNodeFilteringChanged();
+}
+
+////////////////////////////////////////////////////////////////////////////////////////////////////
+// Node Filtering
+////////////////////////////////////////////////////////////////////////////////////////////////////
+
+void STableTreeView::InitNodeFiltering()
+{
+	// Filter Configurator is created when the Filter Configurator button is pressed.
+	//FilterConfigurator = nullptr;
+	//CurrentAsyncOpFilterConfigurator = nullptr;
+	//OnFilterChangesCommittedHandle = nullptr;
+}
+
+////////////////////////////////////////////////////////////////////////////////////////////////////
+
+void STableTreeView::OnNodeFilteringChanged()
+{
 	if (!bIsUpdateRunning)
 	{
-		OnPreAsyncUpdate();
-
-		FGraphEventRef CompletedEvent = StartCreateGroupsTask();
-		CompletedEvent = StartSortTreeNodesTask(CompletedEvent);
-		InProgressAsyncOperationEvent = StartApplyFiltersTask(CompletedEvent);
+		ScheduleNodeFilteringAsyncOperation();
 	}
 	else
 	{
-		CancelCurrentAsyncOp();
-	}
-}
-
-////////////////////////////////////////////////////////////////////////////////////////////////////
-// Filtering
-////////////////////////////////////////////////////////////////////////////////////////////////////
-
-void STableTreeView::OnFilteringChanged()
-{
-	if (!bIsUpdateRunning)
-	{
-		OnPreAsyncUpdate();
-
-		InProgressAsyncOperationEvent = StartApplyFiltersTask();
-	}
-	else
-	{
@@ -999,35 +950,343 @@
 
 ////////////////////////////////////////////////////////////////////////////////////////////////////
 
-void STableTreeView::ApplyFiltering()
+bool STableTreeView::ScheduleNodeFilteringAsyncOperationIfNeeded()
 {
-	ApplyFiltering(Root);
+	if (HasInProgressAsyncOperation(EAsyncOperationType::NodeFiltering) ||
+		(FilterConfigurator.IsValid() &&
+		 CurrentAsyncOpFilterConfigurator &&
+		 *FilterConfigurator != *CurrentAsyncOpFilterConfigurator))
+	{
+		ScheduleNodeFilteringAsyncOperation();
+		return true;
+	}
+
+	return false;
 }
 
 ////////////////////////////////////////////////////////////////////////////////////////////////////
 
-void STableTreeView::ApplyFiltering(FTableTreeNodePtr SubTreeRootPtr)
+void STableTreeView::ScheduleNodeFilteringAsyncOperation()
+{
+	OnPreAsyncUpdate();
+
+	FGraphEventRef CompletedEvent = StartNodeFilteringTask();
+	CompletedEvent = StartGroupingTask(CompletedEvent);
+	CompletedEvent = StartSortingTask(CompletedEvent);
+	InProgressAsyncOperationEvent = StartHierarchyFilteringTask(CompletedEvent);
+}
+
+////////////////////////////////////////////////////////////////////////////////////////////////////
+
+FGraphEventRef STableTreeView::StartNodeFilteringTask(FGraphEventRef Prerequisite)
+{
+	AddInProgressAsyncOperation(EAsyncOperationType::NodeFiltering);
+
+	if (FilterConfigurator.IsValid() && CurrentAsyncOpFilterConfigurator)
+	{
+		*CurrentAsyncOpFilterConfigurator = *FilterConfigurator;
+	}
+
+	FGraphEventArray Prerequisites;
+	if (Prerequisite.IsValid())
+	{
+		Prerequisites.Add(Prerequisite);
+	}
+	else
+	{
+		Prerequisites.Add(DispatchEvent);
+	}
+
+	return TGraphTask<FTableTreeViewNodeFilteringAsyncTask>::CreateTask(&Prerequisites).ConstructAndDispatchWhenReady(this);
+}
+
+////////////////////////////////////////////////////////////////////////////////////////////////////
+
+// Note: This function is called from a task thread!
+void STableTreeView::ApplyNodeFiltering()
 {
 	FStopwatch Stopwatch;
 	Stopwatch.Start();
 
-	// Apply filter to all groups and its children.
-	if (FilterConfigurator_HasFilters())
+	// TableRowNodes --> FilteredNodes / FilteredNodesPtr
+	FilteredNodes.Reset();
+	if (CurrentAsyncOpFilterConfigurator && !CurrentAsyncOpFilterConfigurator->IsEmpty())
 	{
-		ApplyAdvancedFiltersForNode(SubTreeRootPtr);
+		for (const FTableTreeNodePtr& NodePtr : TableRowNodes)
+		{
+			if (AsyncOperationProgress.ShouldCancelAsyncOp())
+			{
+				break;
+			}
+
+			if (FilterNode(*CurrentAsyncOpFilterConfigurator, *NodePtr))
+			{
+				FilteredNodes.Add(NodePtr);
+			}
+		}
+		FilteredNodesPtr = &FilteredNodes;
 	}
 	else
 	{
-		const bool bFilterIsEmpty = CurrentAsyncOpTextFilter->GetRawFilterText().IsEmpty();
-		ApplyHierarchicalFilterForNode(SubTreeRootPtr, bFilterIsEmpty);
+		FilteredNodesPtr = &TableRowNodes;
 	}
 
-	FilteredGroupNodes.Reset();
+	Stopwatch.Stop();
+	const double FilteringTime = Stopwatch.GetAccumulatedTime();
+	if (FilteringTime > 0.1)
+	{
+		UE_LOG(LogInsights, Log, TEXT("[Tree - %s] Node filtering completed in %.3fs."), *Table->GetDisplayName().ToString(), FilteringTime);
+	}
+}
+
+////////////////////////////////////////////////////////////////////////////////////////////////////
+
+bool STableTreeView::FilterNode(const FFilterConfigurator& InFilterConfigurator, const FTableTreeNode& InNode) const
+{
+	UpdateFilterContext(InFilterConfigurator, InNode);
+	return InFilterConfigurator.ApplyFilters(FilterContext);
+}
+
+////////////////////////////////////////////////////////////////////////////////////////////////////
+
+void STableTreeView::InitFilterConfigurator(FFilterConfigurator& InOutFilterConfigurator)
+{
+	for (const TSharedRef<FTableColumn>& Column : Table->GetColumns())
+	{
+		if (!Column->CanBeFiltered())
+		{
+			continue;
+		}
+
+		switch (Column->GetDataType())
+		{
+			case ETableCellDataType::Int64:
+			{
+				TSharedRef<FFilter> Filter = MakeShared<FFilter>(
+					Column->GetIndex(),
+					Column->GetTitleName(),
+					Column->GetDescription(),
+					EFilterDataType::Int64,
+					Column->GetValueConverter(),
+					FFilterService::Get()->GetIntegerOperators()
+				);
+				InOutFilterConfigurator.Add(Filter);
+				FilterContext.AddFilterData<int64>(Column->GetIndex(), 0);
+				break;
+			}
+
+			case ETableCellDataType::Double:
+			{
+				TSharedRef<FFilter> Filter = MakeShared<FFilter>(
+					Column->GetIndex(),
+					Column->GetTitleName(),
+					Column->GetDescription(),
+					EFilterDataType::Double,
+					Column->GetValueConverter(),
+					FFilterService::Get()->GetDoubleOperators());
+				InOutFilterConfigurator.Add(Filter);
+				FilterContext.AddFilterData<double>(Column->GetIndex(), 0.0);
+				break;
+			}
+
+			case ETableCellDataType::CString:
+			case ETableCellDataType::Text:
+			case ETableCellDataType::Custom:
+			{
+				if (!Column->IsHierarchy())
+				{
+					TSharedRef<FFilter> Filter = MakeShared<FFilter>(
+						Column->GetIndex(),
+						Column->GetTitleName(),
+						Column->GetDescription(),
+						EFilterDataType::String,
+						Column->GetValueConverter(),
+						FFilterService::Get()->GetStringOperators());
+					InOutFilterConfigurator.Add(Filter);
+					FilterContext.AddFilterData<FString>(Column->GetIndex(), FString());
+				}
+				break;
+			}
+		}
+	}
+}
+
+////////////////////////////////////////////////////////////////////////////////////////////////////
+
+void STableTreeView::UpdateFilterContext(const FFilterConfigurator& InFilterConfigurator, const FTableTreeNode& InNode) const
+{
+	for (const TSharedRef<FTableColumn>& Column : Table->GetColumns())
+	{
+		if (!Column->CanBeFiltered() || !InFilterConfigurator.IsKeyUsed(Column->GetIndex()))
+		{
+			continue;
+		}
+
+		switch (Column->GetDataType())
+		{
+			case ETableCellDataType::Int64:
+			{
+				FilterContext.SetFilterData<int64>(Column->GetIndex(), Column->GetValue(InNode)->AsInt64());
+				break;
+			}
+
+			case ETableCellDataType::Double:
+			{
+				FilterContext.SetFilterData<double>(Column->GetIndex(), Column->GetValue(InNode)->AsDouble());
+				break;
+			}
+
+			case ETableCellDataType::CString:
+			case ETableCellDataType::Text:
+			case ETableCellDataType::Custom:
+			{
+				if (!Column->IsHierarchy())
+				{
+					FilterContext.SetFilterData<FString>(Column->GetIndex(), Column->GetValue(InNode)->AsString());
+				}
+				break;
+			}
+		}
+	}
+}
+
+////////////////////////////////////////////////////////////////////////////////////////////////////
+
+BEGIN_SLATE_FUNCTION_BUILD_OPTIMIZATION
+TSharedRef<SWidget> STableTreeView::ConstructFilterConfiguratorButton()
+{
+	return SNew(SButton)
+		.ContentPadding(FMargin(-2.0f, 2.0f, -2.0f, 2.0f))
+		.ToolTipText(LOCTEXT("FilterConfiguratorBtn_ToolTip", "Opens the filter configurator window."))
+		.OnClicked(this, &STableTreeView::FilterConfigurator_OnClicked)
+		[
+			SNew(SImage)
+			.Image(FInsightsStyle::GetBrush("Icons.ClassicFilterConfig"))
+		];
+}
+END_SLATE_FUNCTION_BUILD_OPTIMIZATION
+
+////////////////////////////////////////////////////////////////////////////////////////////////////
+
+FReply STableTreeView::FilterConfigurator_OnClicked()
+{
+	if (!FilterConfigurator.IsValid())
+	{
+		FilterConfigurator = MakeShared<FFilterConfigurator>();
+		InitFilterConfigurator(*FilterConfigurator);
+		OnFilterChangesCommittedHandle = FilterConfigurator->GetOnChangesCommittedEvent().AddSP(this, &STableTreeView::OnNodeFilteringChanged);
+		CurrentAsyncOpFilterConfigurator = new FFilterConfigurator(*FilterConfigurator);
+	}
+
+	FFilterService::Get()->CreateFilterConfiguratorWidget(FilterConfigurator);
+
+	return FReply::Handled();
+}
+
+////////////////////////////////////////////////////////////////////////////////////////////////////
+// Hierarchy Filtering
+////////////////////////////////////////////////////////////////////////////////////////////////////
+
+void STableTreeView::InitHierarchyFiltering()
+{
+	// The text filter editable on UI thread.
+	TextFilter = MakeShared<FTableTreeNodeTextFilter>(FTableTreeNodeTextFilter::FItemToStringArray::CreateStatic(&STableTreeView::HandleItemToStringArray));
+
+	// A copy of the TextFilter used to actually apply the filtering in an async task.
+	CurrentAsyncOpTextFilter = MakeShared<FTableTreeNodeTextFilter>(FTableTreeNodeTextFilter::FItemToStringArray::CreateStatic(&STableTreeView::HandleItemToStringArray));
+
+	Filters = MakeShared<FTableTreeNodeFilterCollection>();
+	Filters->Add(CurrentAsyncOpTextFilter);
+}
+
+////////////////////////////////////////////////////////////////////////////////////////////////////
+
+void STableTreeView::HandleItemToStringArray(const FTableTreeNodePtr& FTableTreeNodePtr, TArray<FString>& OutSearchStrings)
+{
+	OutSearchStrings.Add(FTableTreeNodePtr->GetName().ToString());
+}
+
+////////////////////////////////////////////////////////////////////////////////////////////////////
+
+void STableTreeView::OnHierarchyFilteringChanged()
+{
+	if (!bIsUpdateRunning)
+	{
+		ScheduleHierarchyFilteringAsyncOperation();
+	}
+	else
+	{
+		CancelCurrentAsyncOp();
+	}
+}
+
+////////////////////////////////////////////////////////////////////////////////////////////////////
+
+bool STableTreeView::ScheduleHierarchyFilteringAsyncOperationIfNeeded()
+{
+	if (HasInProgressAsyncOperation(EAsyncOperationType::HierarchyFiltering) ||
+		TextFilter->GetRawFilterText().CompareTo(CurrentAsyncOpTextFilter->GetRawFilterText()) != 0)
+	{
+		ScheduleHierarchyFilteringAsyncOperation();
+		return true;
+	}
+
+	return false;
+}
+
+////////////////////////////////////////////////////////////////////////////////////////////////////
+
+void STableTreeView::ScheduleHierarchyFilteringAsyncOperation()
+{
+	OnPreAsyncUpdate();
+
+	InProgressAsyncOperationEvent = StartHierarchyFilteringTask();
+}
+
+////////////////////////////////////////////////////////////////////////////////////////////////////
+
+FGraphEventRef STableTreeView::StartHierarchyFilteringTask(FGraphEventRef Prerequisite)
+{
+	AddInProgressAsyncOperation(EAsyncOperationType::HierarchyFiltering);
+
+	CurrentAsyncOpTextFilter->SetRawFilterText(TextFilter->GetRawFilterText());
+
+	FGraphEventArray Prerequisites;
+	if (Prerequisite.IsValid())
+	{
+		Prerequisites.Add(Prerequisite);
+	}
+	else
+	{
+		Prerequisites.Add(DispatchEvent);
+	}
+
+	return TGraphTask<FTableTreeViewHierarchyFilteringAsyncTask>::CreateTask(&Prerequisites).ConstructAndDispatchWhenReady(this);
+}
+
+////////////////////////////////////////////////////////////////////////////////////////////////////
+
+// Note: This function is called from a task thread!
+void STableTreeView::ApplyHierarchyFiltering()
+{
+	FStopwatch Stopwatch;
+	Stopwatch.Start();
+
+	if (CurrentAsyncOpTextFilter->GetRawFilterText().IsEmpty())
+	{
+		ApplyEmptyHierarchyFilteringRec(Root);
+	}
+	else
+	{
+		ApplyHierarchyFilteringRec(Root);
+	}
+
+	// The Root node is always hidden. The tree shows the filtered children of the Root node.
 	const TArray<FBaseTreeNodePtr>& RootChildren = Root->GetFilteredChildren();
 	const int32 NumRootChildren = RootChildren.Num();
+	FilteredGroupNodes.Reset(NumRootChildren);
 	for (int32 Cx = 0; Cx < NumRootChildren; ++Cx)
 	{
-		// Add a child.
 		const FTableTreeNodePtr& ChildNodePtr = StaticCastSharedPtr<FTableTreeNode>(RootChildren[Cx]);
 		if (ChildNodePtr->IsGroup())
 		{
@@ -1035,45 +1294,51 @@
 		}
 	}
 
-	TreeViewBannerText = FText::GetEmpty();
-	if (Root->GetChildren().Num() == 0)
-	{
-		TreeViewBannerText = LOCTEXT("EmptyTreeMsg", "Tree is empty.");
-	}
-	else if (Root->GetFilteredChildren().Num() == 0)
-	{
-		TreeViewBannerText = LOCTEXT("AllNodesFiltered", "All tree nodes are filtered out.");
-	}
+
+	UpdateBannerText();
 
 	Stopwatch.Stop();
 	const double FilteringTime = Stopwatch.GetAccumulatedTime();
 	if (FilteringTime > 0.1)
 	{
-		UE_LOG(LogInsights, Log, TEXT("[Tree - %s] Filtering completed in %.3fs."), *Table->GetDisplayName().ToString(), FilteringTime);
+		UE_LOG(LogInsights, Log, TEXT("[Tree - %s] Hierarchy filtering completed in %.3fs."), *Table->GetDisplayName().ToString(), FilteringTime);
 	}
-
-	UpdateAggregatedValuesRec(*Root);
 }
 
 ////////////////////////////////////////////////////////////////////////////////////////////////////
 
-bool STableTreeView::ApplyHierarchicalFilterForNode(FTableTreeNodePtr NodePtr, bool bFilterIsEmpty)
+void STableTreeView::ApplyEmptyHierarchyFilteringRec(FTableTreeNodePtr NodePtr)
 {
-	bool bIsNodeVisible = bFilterIsEmpty || Filters->PassesAllFilters(NodePtr);
+	if (NodePtr->IsGroup())
+	{
+		// If a group node passes the filter, all child nodes will be shown.
+		MakeSubtreeVisible(NodePtr, true);
+	}
+	else
+	{
+		NodePtr->SetIsFiltered(false);
+	}
+}
+
+////////////////////////////////////////////////////////////////////////////////////////////////////
+
+bool STableTreeView::ApplyHierarchyFilteringRec(FTableTreeNodePtr NodePtr)
+{
+	bool bIsNodeVisible = Filters->PassesAllFilters(NodePtr);
 
 	if (NodePtr->IsGroup())
 	{
 		// If a group node passes the filter, all child nodes will be shown.
 		if (bIsNodeVisible)
 		{
-			MakeSubtreeVisible(NodePtr, bFilterIsEmpty);
+			MakeSubtreeVisible(NodePtr, false);
 			return true;
 		}
 
 		const TArray<FBaseTreeNodePtr>& GroupChildren = NodePtr->GetChildren();
 		const int32 NumChildren = GroupChildren.Num();
 
-		NodePtr->ClearFilteredChildren(bFilterIsEmpty ? NumChildren : 0);
+		NodePtr->ClearFilteredChildren(0);
 
 		int32 NumVisibleChildren = 0;
 		for (int32 Cx = 0; Cx < NumChildren; ++Cx)
@@ -1085,7 +1350,7 @@
 
 			// Add a child.
 			const FTableTreeNodePtr& ChildNodePtr = StaticCastSharedPtr<FTableTreeNode>(GroupChildren[Cx]);
-			if (ApplyHierarchicalFilterForNode(ChildNodePtr, bFilterIsEmpty))
+			if (ApplyHierarchyFilteringRec(ChildNodePtr))
 			{
 				NodePtr->AddFilteredChild(ChildNodePtr);
 				NumVisibleChildren++;
@@ -1093,64 +1358,14 @@
 		}
 
 		const bool bIsGroupNodeVisible = bIsNodeVisible || NumVisibleChildren > 0;
-		if (!bFilterIsEmpty && bIsGroupNodeVisible)
-		{
-			NodesToExpand.Add(NodePtr);
-		}
-		return bIsGroupNodeVisible;
-	}
-	else
-	{
-		NodePtr->SetIsFiltered(!bIsNodeVisible);
-		return bIsNodeVisible;
-	}
-}
-
-////////////////////////////////////////////////////////////////////////////////////////////////////
-
-bool STableTreeView::ApplyAdvancedFiltersForNode(FTableTreeNodePtr NodePtr)
-{
-	if (NodePtr->IsGroup())
-	{
-		const TArray<FBaseTreeNodePtr>& GroupChildren = NodePtr->GetChildren();
-		const int32 NumChildren = GroupChildren.Num();
-
-		NodePtr->ClearFilteredChildren();
-
-		int32 NumVisibleChildren = 0;
-		for (int32 Cx = 0; Cx < NumChildren; ++Cx)
-		{
-			if (AsyncOperationProgress.ShouldCancelAsyncOp())
-			{
-				break;
-			}
-
-			// Add a child.
-			const FTableTreeNodePtr& ChildNodePtr = StaticCastSharedPtr<FTableTreeNode>(GroupChildren[Cx]);
-			if (ApplyAdvancedFiltersForNode(ChildNodePtr))
-			{
-				NodePtr->AddFilteredChild(ChildNodePtr);
-				NumVisibleChildren++;
-			}
-		}
-
-		const bool bIsGroupNodeVisible = NumVisibleChildren > 0;
-
 		if (bIsGroupNodeVisible)
 		{
-			// Add a group.
-			NodePtr->SetExpansion(true);
+			NodesToExpand.Add(NodePtr);
 		}
-		else
-		{
-			NodePtr->SetExpansion(false);
-		}
-
 		return bIsGroupNodeVisible;
 	}
 	else
 	{
-		bool bIsNodeVisible = ApplyAdvancedFilters(NodePtr);
 		NodePtr->SetIsFiltered(!bIsNodeVisible);
 		return bIsNodeVisible;
 	}
@@ -1196,10 +1411,28 @@
 }
 
 ////////////////////////////////////////////////////////////////////////////////////////////////////
+// SearchBox
+////////////////////////////////////////////////////////////////////////////////////////////////////
 
-void STableTreeView::HandleItemToStringArray(const FTableTreeNodePtr& FTableTreeNodePtr, TArray<FString>& OutSearchStrings)
+BEGIN_SLATE_FUNCTION_BUILD_OPTIMIZATION
+TSharedRef<SWidget> STableTreeView::ConstructSearchBox()
 {
-	OutSearchStrings.Add(FTableTreeNodePtr->GetName().ToString());
+	SAssignNew(SearchBox, SSearchBox)
+		.HintText(LOCTEXT("SearchBox_Hint", "Search"))
+		.ToolTipText(LOCTEXT("SearchBox_ToolTip", "Type here to search the tree hierarchy by item or group name."))
+		.OnTextChanged(this, &STableTreeView::SearchBox_OnTextChanged);
+
+	return SearchBox.ToSharedRef();
+}
+END_SLATE_FUNCTION_BUILD_OPTIMIZATION
+
+////////////////////////////////////////////////////////////////////////////////////////////////////
+
+void STableTreeView::SearchBox_OnTextChanged(const FText& InFilterText)
+{
+	TextFilter->SetRawFilterText(InFilterText);
+	SearchBox->SetError(TextFilter->GetFilterErrorText());
+	OnHierarchyFilteringChanged();
 }
 
 ////////////////////////////////////////////////////////////////////////////////////////////////////
@@ -1232,16 +1465,26 @@
 {
 	if (InParent->OnLazyCreateChildren(SharedThis(this)))
 	{
-		// The filtering is ignored for lazy expanded nodes.
-		InParent->ResetFilteredChildrenRec();
+		// Node filtering does not apply to lazy expanded nodes.
+		// Grouping is ignored for lazy expanded nodes.
 
+		// Update aggregation.
 		UpdateAggregatedValuesRec(*InParent);
-		FTableTreeNodePtr Node = StaticCastSharedPtr<FTableTreeNode>(InParent->GetParentNode());
+		FTableTreeNodePtr Node = StaticCastSharedPtr<FTableTreeNode>(InParent->GetParent());
 		while (Node)
 		{
 			UpdateAggregatedValuesSingleNode(*Node);
-			Node = StaticCastSharedPtr<FTableTreeNode>(Node->GetParentNode());
+			Node = StaticCastSharedPtr<FTableTreeNode>(Node->GetParent());
 		}
+
+		// Update sorting.
+		if (CurrentSorter.IsValid())
+		{
+			SortTreeNodesRec(*InParent, *CurrentSorter, ColumnSortMode);
+		}
+
+		// The hierarchy filtering is ignored for lazy expanded nodes.
+		InParent->ResetFilteredChildrenRec();
 	}
 
 	const TArray<FBaseTreeNodePtr>& FilteredChildren = InParent->GetFilteredChildren();
@@ -1339,36 +1582,6 @@
 }
 
 ////////////////////////////////////////////////////////////////////////////////////////////////////
-// SearchBox
-////////////////////////////////////////////////////////////////////////////////////////////////////
-
-void STableTreeView::SearchBox_OnTextChanged(const FText& InFilterText)
-{
-	TextFilter->SetRawFilterText(InFilterText);
-	SearchBox->SetError(TextFilter->GetFilterErrorText());
-	OnFilteringChanged();
-}
-
-////////////////////////////////////////////////////////////////////////////////////////////////////
-
-bool STableTreeView::SearchBox_IsEnabled() const
-{
-	return !FilterConfigurator_HasFilters();
-}
-
-////////////////////////////////////////////////////////////////////////////////////////////////////
-
-FText STableTreeView::SearchBox_GetTooltipText() const
-{
-	if (SearchBox_IsEnabled())
-	{
-		return LOCTEXT("SearchBox_ToolTip", "Type here to search the tree hierarchy by item or group name.");
-	}
-
-	return LOCTEXT("SearchBox_Disabled_ToolTip", "Searching the tree hierarchy is disabled when advanced filters are set.");
-}
-
-////////////////////////////////////////////////////////////////////////////////////////////////////
 // Grouping
 ////////////////////////////////////////////////////////////////////////////////////////////////////
 
@@ -1376,11 +1589,7 @@
 {
 	if (!bIsUpdateRunning)
 	{
-		OnPreAsyncUpdate();
-
-		FGraphEventRef CompletedEvent = StartCreateGroupsTask();
-		CompletedEvent = StartSortTreeNodesTask(CompletedEvent);
-		InProgressAsyncOperationEvent = StartApplyFiltersTask(CompletedEvent);
+		ScheduleGroupingAsyncOperation();
 	}
 	else
 	{
@@ -1390,17 +1599,75 @@
 
 ////////////////////////////////////////////////////////////////////////////////////////////////////
 
-void STableTreeView::ApplyGrouping()
+bool STableTreeView::ScheduleGroupingAsyncOperationIfNeeded()
 {
-	CreateGroups(CurrentGroupings);
+	bool bGroupingHasChanged = HasInProgressAsyncOperation(EAsyncOperationType::Grouping);
+
+	if (!bGroupingHasChanged &&
+		CurrentGroupings.Num() != CurrentAsyncOpGroupings.Num())
+	{
+		bGroupingHasChanged = true;
+	}
+
+	if (!bGroupingHasChanged)
+	{
+		for (int32 Index = 0; Index < CurrentGroupings.Num(); ++Index)
+		{
+			if (CurrentGroupings[Index] != CurrentAsyncOpGroupings[Index])
+			{
+				bGroupingHasChanged = true;
+				break;
+			}
+		}
+	}
+
+	if (bGroupingHasChanged)
+	{
+		ScheduleGroupingAsyncOperation();
+		return true;
+	}
+
+	return false;
 }
 
 ////////////////////////////////////////////////////////////////////////////////////////////////////
 
-// Note: This function is also called from async thread, with CurrentAsyncOpGroupings as parameter.
+void STableTreeView::ScheduleGroupingAsyncOperation()
+{
+	OnPreAsyncUpdate();
+
+	FGraphEventRef CompletedEvent = StartGroupingTask();
+	CompletedEvent = StartSortingTask(CompletedEvent);
+	InProgressAsyncOperationEvent = StartHierarchyFilteringTask(CompletedEvent);
+}
+
+////////////////////////////////////////////////////////////////////////////////////////////////////
+
+FGraphEventRef STableTreeView::StartGroupingTask(FGraphEventRef Prerequisite)
+{
+	AddInProgressAsyncOperation(EAsyncOperationType::Grouping);
+
+	CurrentAsyncOpGroupings = CurrentGroupings;
+
+	FGraphEventArray Prerequisites;
+	if (Prerequisite.IsValid())
+	{
+		Prerequisites.Add(Prerequisite);
+	}
+	else
+	{
+		Prerequisites.Add(DispatchEvent);
+	}
+
+	return TGraphTask<FTableTreeViewGroupingAsyncTask>::CreateTask(&Prerequisites).ConstructAndDispatchWhenReady(this, &CurrentAsyncOpGroupings);
+}
+
+////////////////////////////////////////////////////////////////////////////////////////////////////
+
+// Note: This function is called from a task thread!
 void STableTreeView::CreateGroups(const TArray<TSharedPtr<FTreeNodeGrouping>>& Groupings)
 {
-	if (TableTreeNodes.Num() == 0)
+	if (FilteredNodesPtr->Num() == 0)
 	{
 		Root->ClearChildren();
 		return;
@@ -1411,7 +1678,7 @@
 
 	if (Groupings.Num() > 0)
 	{
-		GroupNodesRec(TableTreeNodes, *Root, 0, Groupings);
+		GroupNodesRec(*FilteredNodesPtr, *Root, 0, Groupings);
 	}
 
 	Stopwatch.Stop();
@@ -1489,250 +1756,3 @@
 }
 
 ////////////////////////////////////////////////////////////////////////////////////////////////////
-
-void STableTreeView::UpdateCStringSameValueAggregationSingleNode(const FTableColumn& InColumn, FTableTreeNode& InOutGroupNode)
-{
-	check(InOutGroupNode.IsGroup());
-
-	const TCHAR* AggregatedValue = nullptr;
-
-	// Find the first child node.
-	for (FBaseTreeNodePtr NodePtr : InOutGroupNode.GetFilteredChildren())
-	{
-		if (!NodePtr->IsGroup())
-		{
-			const TOptional<FTableCellValue> NodeValue = InColumn.GetValue(*NodePtr);
-			if (NodeValue.IsSet() &&
-				NodeValue.GetValue().DataType == ETableCellDataType::CString)
-			{
-				AggregatedValue = NodeValue.GetValue().CString;
-			}
-		}
-		else
-		{
-			FTableTreeNode& TableNode = *(FTableTreeNode*)NodePtr.Get();
-			if (TableNode.HasAggregatedValue(InColumn.GetId()))
-			{
-				const FTableCellValue& ChildGroupAggregatedValue = TableNode.GetAggregatedValue(InColumn.GetId());
-				if (ChildGroupAggregatedValue.DataType == ETableCellDataType::CString)
-				{
-					AggregatedValue = ChildGroupAggregatedValue.CString;
-				}
-			}
-		}
-
-		break;
-	}
-
-	if (AggregatedValue != nullptr)
-	{
-		// Check if all other children have the same value as the first node.
-		for (FBaseTreeNodePtr NodePtr : InOutGroupNode.GetFilteredChildren())
-		{
-			if (!NodePtr->IsGroup())
-			{
-				const TOptional<FTableCellValue> NodeValue = InColumn.GetValue(*NodePtr);
-				if (NodeValue.IsSet() &&
-					NodeValue.GetValue().DataType == ETableCellDataType::CString &&
-					AggregatedValue == NodeValue.GetValue().CString)
-				{
-					continue;
-				}
-			}
-			else
-			{
-				FTableTreeNode& TableNode = *(FTableTreeNode*)NodePtr.Get();
-				if (TableNode.HasAggregatedValue(InColumn.GetId()))
-				{
-					const FTableCellValue& ChildGroupAggregatedValue = TableNode.GetAggregatedValue(InColumn.GetId());
-					if (ChildGroupAggregatedValue.DataType == ETableCellDataType::CString &&
-						AggregatedValue == ChildGroupAggregatedValue.CString)
-					{
-						continue;
-					}
-				}
-			}
-
-			AggregatedValue = nullptr;
-			break;
-		}
-	}
-
-	if (AggregatedValue != nullptr)
-	{
-		InOutGroupNode.SetAggregatedValue(InColumn.GetId(), FTableCellValue(AggregatedValue));
-	}
-	else
-	{
-		InOutGroupNode.ResetAggregatedValue(InColumn.GetId());
-	}
-}
-
-////////////////////////////////////////////////////////////////////////////////////////////////////
-
-void STableTreeView::UpdateCStringSameValueAggregationRec(const FTableColumn& InColumn, FTableTreeNode& InOutGroupNode)
-{
-	check(InOutGroupNode.IsGroup());
-
-	// Update child group nodes first.
-	for (FBaseTreeNodePtr NodePtr : InOutGroupNode.GetFilteredChildren())
-	{
-		if (NodePtr->IsGroup())
-		{
-			FTableTreeNode& TableNode = *(FTableTreeNode*)NodePtr.Get();
-			STableTreeView::UpdateCStringSameValueAggregationRec(InColumn, TableNode);
-		}
-	}
-
-	STableTreeView::UpdateCStringSameValueAggregationSingleNode(InColumn, InOutGroupNode);
-}
-
-////////////////////////////////////////////////////////////////////////////////////////////////////
-
-template<typename T, bool bSetInitialValue, bool bIsRercursive>
-void STableTreeView::UpdateAggregation(const FTableColumn& InColumn, FTableTreeNode& InOutGroupNode, const T InitialAggregatedValue, TFunctionRef<T(T, const FTableCellValue&)> ValueGetterFunc)
-{
-	T AggregatedValue = InitialAggregatedValue;
-
-	for (FBaseTreeNodePtr NodePtr : InOutGroupNode.GetFilteredChildren())
-	{
-		if (NodePtr->IsGroup())
-		{
-			FTableTreeNode& TableNode = *(FTableTreeNode*)NodePtr.Get();
-			if (bIsRercursive)
-			{
-				TableNode.ResetAggregatedValue(InColumn.GetId());
-				STableTreeView::UpdateAggregation<T, bSetInitialValue, true>(InColumn, TableNode, InitialAggregatedValue, ValueGetterFunc);
-			}
-			if (TableNode.HasAggregatedValue(InColumn.GetId()))
-			{
-				AggregatedValue = ValueGetterFunc(AggregatedValue, TableNode.GetAggregatedValue(InColumn.GetId()));
-			}
-		}
-		else
-		{
-			const TOptional<FTableCellValue> NodeValue = InColumn.GetValue(*NodePtr);
-			if (NodeValue.IsSet())
-			{
-				AggregatedValue = ValueGetterFunc(AggregatedValue, NodeValue.GetValue());
-			}
-		}
-	}
-
-	if (bSetInitialValue || AggregatedValue != InitialAggregatedValue)
-	{
-		InOutGroupNode.SetAggregatedValue(InColumn.GetId(), FTableCellValue(AggregatedValue));
-	}
-}
-
-////////////////////////////////////////////////////////////////////////////////////////////////////
-
-template<bool bIsRercursive>
-void STableTreeView::UpdateAggregatedValues(TSharedPtr<FTable> InTable, FTableTreeNode& InOutGroupNode)
-{
-	FStopwatch Stopwatch;
-	Stopwatch.Start();
-
-	check(InOutGroupNode.IsGroup());
-
-	for (const TSharedRef<FTableColumn>& ColumnRef : InTable->GetColumns())
-	{
-		FTableColumn& Column = ColumnRef.Get();
-		switch (Column.GetAggregation())
-		{
-			case ETableColumnAggregation::Sum:
-				if (Column.GetDataType() == ETableCellDataType::Float || Column.GetDataType() == ETableCellDataType::Double)
-				{
-					STableTreeView::UpdateAggregation<double, true, bIsRercursive>(Column, InOutGroupNode, 0.0,
-						[](double InValue, const FTableCellValue& InTableCellValue)
-						{
-							return InValue + InTableCellValue.AsDouble();
-						});
-				}
-				else
-				{
-					STableTreeView::UpdateAggregation<int64, true, bIsRercursive>(Column, InOutGroupNode, (int64)0,
-						[](int64 InValue, const FTableCellValue& InTableCellValue)
-						{
-							return InValue + InTableCellValue.AsInt64();
-						});
-				}
-				break;
-
-			case ETableColumnAggregation::Min:
-				if (Column.GetDataType() == ETableCellDataType::Float || Column.GetDataType() == ETableCellDataType::Double)
-				{
-					STableTreeView::UpdateAggregation<double, false, bIsRercursive>(Column, InOutGroupNode, std::numeric_limits<double>::max(),
-						[](double InValue, const FTableCellValue& InTableCellValue)
-						{
-							return FMath::Min(InValue, InTableCellValue.AsDouble());
-						});
-				}
-				else
-				{
-					STableTreeView::UpdateAggregation<int64, false, bIsRercursive>(Column, InOutGroupNode, std::numeric_limits<int64>::max(),
-						[](int64 InValue, const FTableCellValue& InTableCellValue)
-						{
-							return FMath::Min(InValue, InTableCellValue.AsInt64());
-						});
-				}
-				break;
-
-			case ETableColumnAggregation::Max:
-				if (Column.GetDataType() == ETableCellDataType::Float || Column.GetDataType() == ETableCellDataType::Double)
-				{
-					STableTreeView::UpdateAggregation<double, false, bIsRercursive>(Column, InOutGroupNode, std::numeric_limits<double>::lowest(),
-						[](double InValue, const FTableCellValue& InTableCellValue)
-						{
-							return FMath::Max(InValue, InTableCellValue.AsDouble());
-						});
-				}
-				else
-				{
-					STableTreeView::UpdateAggregation<int64, false, bIsRercursive>(Column, InOutGroupNode, std::numeric_limits<int64>::min(),
-						[](int64 InValue, const FTableCellValue& InTableCellValue)
-						{
-							return FMath::Max(InValue, InTableCellValue.AsInt64());
-						});
-				}
-				break;
-
-			case ETableColumnAggregation::SameValue:
-				if (Column.GetDataType() == ETableCellDataType::CString)
-				{
-					if (bIsRercursive)
-					{
-						STableTreeView::UpdateCStringSameValueAggregationRec(Column, InOutGroupNode);
-					}
-					else
-					{
-						STableTreeView::UpdateCStringSameValueAggregationSingleNode(Column, InOutGroupNode);
-					}
-				}
-				break;
-		}
-	}
-
-	Stopwatch.Stop();
-	const double AggregationTime = Stopwatch.GetAccumulatedTime();
-	if (AggregationTime > 0.1)
-	{
-		UE_LOG(LogInsights, Log, TEXT("[Tree - %s] Aggregation completed in %.3fs."), *InTable->GetDisplayName().ToString(), AggregationTime);
-	}
-}
-
-////////////////////////////////////////////////////////////////////////////////////////////////////
-
-void STableTreeView::UpdateAggregatedValuesSingleNode(FTableTreeNode& InOutGroupNode)
-{
-	STableTreeView::UpdateAggregatedValues<false>(Table, InOutGroupNode);
-}
-
-////////////////////////////////////////////////////////////////////////////////////////////////////
-
-void STableTreeView::UpdateAggregatedValuesRec(FTableTreeNode& InOutGroupNode)
-{
-	STableTreeView::UpdateAggregatedValues<true>(Table, InOutGroupNode);
-}
-
-////////////////////////////////////////////////////////////////////////////////////////////////////
@@ -2308,11 +2328,219 @@
 }
 
 ////////////////////////////////////////////////////////////////////////////////////////////////////
+// Aggregation
+////////////////////////////////////////////////////////////////////////////////////////////////////
+
+void STableTreeView::UpdateCStringSameValueAggregationSingleNode(const FTableColumn& InColumn, FTableTreeNode& InOutGroupNode)
+{
+	check(InOutGroupNode.IsGroup());
+
+	const TCHAR* AggregatedValue = nullptr;
+
+	// Find the first child node.
+	if (InOutGroupNode.GetChildrenCount() > 0)
+	{
+		FBaseTreeNodePtr NodePtr = InOutGroupNode.GetChildren()[0];
+		const TOptional<FTableCellValue> NodeValue = InColumn.GetValue(*NodePtr);
+		if (NodeValue.IsSet() &&
+			NodeValue.GetValue().DataType == ETableCellDataType::CString)
+		{
+			AggregatedValue = NodeValue.GetValue().CString;
+		}
+	}
+
+	if (AggregatedValue != nullptr)
+	{
+		// Check if all other children have the same value as the first node.
+		for (FBaseTreeNodePtr NodePtr : InOutGroupNode.GetChildren())
+		{
+			const TOptional<FTableCellValue> NodeValue = InColumn.GetValue(*NodePtr);
+			if (NodeValue.IsSet() &&
+				NodeValue.GetValue().DataType == ETableCellDataType::CString &&
+				AggregatedValue == NodeValue.GetValue().CString)
+			{
+				continue;
+			}
+
+			AggregatedValue = nullptr;
+			break;
+		}
+	}
+
+	if (AggregatedValue != nullptr)
+	{
+		InOutGroupNode.SetAggregatedValue(InColumn.GetId(), FTableCellValue(AggregatedValue));
+	}
+	else
+	{
+		InOutGroupNode.ResetAggregatedValue(InColumn.GetId());
+	}
+}
+
+////////////////////////////////////////////////////////////////////////////////////////////////////
+
+void STableTreeView::UpdateCStringSameValueAggregationRec(const FTableColumn& InColumn, FTableTreeNode& InOutGroupNode)
+{
+	check(InOutGroupNode.IsGroup());
+
+	// Update child group nodes first.
+	for (FBaseTreeNodePtr NodePtr : InOutGroupNode.GetChildren())
+	{
+		if (NodePtr->IsGroup())
+		{
+			FTableTreeNode& TableNode = *(FTableTreeNode*)NodePtr.Get();
+			STableTreeView::UpdateCStringSameValueAggregationRec(InColumn, TableNode);
+		}
+	}
+
+	STableTreeView::UpdateCStringSameValueAggregationSingleNode(InColumn, InOutGroupNode);
+}
+
+////////////////////////////////////////////////////////////////////////////////////////////////////
+
+template<typename T, bool bSetInitialValue, bool bIsRercursive>
+void STableTreeView::UpdateAggregation(const FTableColumn& InColumn, FTableTreeNode& InOutGroupNode, const T InitialAggregatedValue, TFunctionRef<T(T, const FTableCellValue&)> ValueGetterFunc)
+{
+	T AggregatedValue = InitialAggregatedValue;
+
+	for (FBaseTreeNodePtr NodePtr : InOutGroupNode.GetChildren())
+	{
+		if (NodePtr->IsGroup())
+		{
+			FTableTreeNode& TableNode = *(FTableTreeNode*)NodePtr.Get();
+			if (bIsRercursive)
+			{
+				TableNode.ResetAggregatedValue(InColumn.GetId());
+				STableTreeView::UpdateAggregation<T, bSetInitialValue, true>(InColumn, TableNode, InitialAggregatedValue, ValueGetterFunc);
+			}
+		}
+		const TOptional<FTableCellValue> NodeValue = InColumn.GetValue(*NodePtr);
+		if (NodeValue.IsSet())
+		{
+			AggregatedValue = ValueGetterFunc(AggregatedValue, NodeValue.GetValue());
+		}
+	}
+
+	if (bSetInitialValue || AggregatedValue != InitialAggregatedValue)
+	{
+		InOutGroupNode.SetAggregatedValue(InColumn.GetId(), FTableCellValue(AggregatedValue));
+	}
+}
+
+////////////////////////////////////////////////////////////////////////////////////////////////////
+
+template<bool bIsRecursive>
+void STableTreeView::UpdateAggregatedValues(TSharedPtr<FTable> InTable, FTableTreeNode& InOutGroupNode)
+{
+	FStopwatch Stopwatch;
+	Stopwatch.Start();
+
+	check(InOutGroupNode.IsGroup());
+
+	for (const TSharedRef<FTableColumn>& ColumnRef : InTable->GetColumns())
+	{
+		FTableColumn& Column = ColumnRef.Get();
+		switch (Column.GetAggregation())
+		{
+			case ETableColumnAggregation::Sum:
+				if (Column.GetDataType() == ETableCellDataType::Float || Column.GetDataType() == ETableCellDataType::Double)
+				{
+					STableTreeView::UpdateAggregation<double, true, bIsRecursive>(Column, InOutGroupNode, 0.0,
+						[](double InValue, const FTableCellValue& InTableCellValue)
+						{
+							return InValue + InTableCellValue.AsDouble();
+						});
+				}
+				else
+				{
+					STableTreeView::UpdateAggregation<int64, true, bIsRecursive>(Column, InOutGroupNode, (int64)0,
+						[](int64 InValue, const FTableCellValue& InTableCellValue)
+						{
+							return InValue + InTableCellValue.AsInt64();
+						});
+				}
+				break;
+
+			case ETableColumnAggregation::Min:
+				if (Column.GetDataType() == ETableCellDataType::Float || Column.GetDataType() == ETableCellDataType::Double)
+				{
+					STableTreeView::UpdateAggregation<double, false, bIsRecursive>(Column, InOutGroupNode, std::numeric_limits<double>::max(),
+						[](double InValue, const FTableCellValue& InTableCellValue)
+						{
+							return FMath::Min(InValue, InTableCellValue.AsDouble());
+						});
+				}
+				else
+				{
+					STableTreeView::UpdateAggregation<int64, false, bIsRecursive>(Column, InOutGroupNode, std::numeric_limits<int64>::max(),
+						[](int64 InValue, const FTableCellValue& InTableCellValue)
+						{
+							return FMath::Min(InValue, InTableCellValue.AsInt64());
+						});
+				}
+				break;
+
+			case ETableColumnAggregation::Max:
+				if (Column.GetDataType() == ETableCellDataType::Float || Column.GetDataType() == ETableCellDataType::Double)
+				{
+					STableTreeView::UpdateAggregation<double, false, bIsRecursive>(Column, InOutGroupNode, std::numeric_limits<double>::lowest(),
+						[](double InValue, const FTableCellValue& InTableCellValue)
+						{
+							return FMath::Max(InValue, InTableCellValue.AsDouble());
+						});
+				}
+				else
+				{
+					STableTreeView::UpdateAggregation<int64, false, bIsRecursive>(Column, InOutGroupNode, std::numeric_limits<int64>::min(),
+						[](int64 InValue, const FTableCellValue& InTableCellValue)
+						{
+							return FMath::Max(InValue, InTableCellValue.AsInt64());
+						});
+				}
+				break;
+
+			case ETableColumnAggregation::SameValue:
+				if (Column.GetDataType() == ETableCellDataType::CString)
+				{
+					if (bIsRecursive)
+					{
+						STableTreeView::UpdateCStringSameValueAggregationRec(Column, InOutGroupNode);
+					}
+					else
+					{
+						STableTreeView::UpdateCStringSameValueAggregationSingleNode(Column, InOutGroupNode);
+					}
+				}
+				break;
+		}
+	}
+
+	Stopwatch.Stop();
+	const double AggregationTime = Stopwatch.GetAccumulatedTime();
+	if (AggregationTime > 0.1)
+	{
+		UE_LOG(LogInsights, Log, TEXT("[Tree - %s] Aggregation completed in %.3fs."), *InTable->GetDisplayName().ToString(), AggregationTime);
+	}
+}
+
+////////////////////////////////////////////////////////////////////////////////////////////////////
+
+void STableTreeView::UpdateAggregatedValuesSingleNode(FTableTreeNode& InOutGroupNode)
+{
+	STableTreeView::UpdateAggregatedValues<false>(Table, InOutGroupNode);
+}
+
+////////////////////////////////////////////////////////////////////////////////////////////////////
+
+void STableTreeView::UpdateAggregatedValuesRec(FTableTreeNode& InOutGroupNode)
+{
+	STableTreeView::UpdateAggregatedValues<true>(Table, InOutGroupNode);
+}
+
+////////////////////////////////////////////////////////////////////////////////////////////////////
 // Sorting
 ////////////////////////////////////////////////////////////////////////////////////////////////////
 
-////////////////////////////////////////////////////////////////////////////////////////////////////
-
 const EColumnSortMode::Type STableTreeView::GetDefaultColumnSortMode()
 {
 	return EColumnSortMode::Descending;
@@ -2362,10 +2590,7 @@
 {
 	if (!bIsUpdateRunning)
 	{
-		OnPreAsyncUpdate();
-
-		FGraphEventRef CompletedEvent = StartSortTreeNodesTask();
-		InProgressAsyncOperationEvent = StartApplyFiltersTask(CompletedEvent);
+		ScheduleSortingAsyncOperation();
 	}
 	else
 	{
@@ -2375,16 +2600,72 @@
 
 ////////////////////////////////////////////////////////////////////////////////////////////////////
 
-void STableTreeView::ApplySorting()
+bool STableTreeView::ScheduleSortingAsyncOperationIfNeeded()
 {
-	if (CurrentSorter.IsValid())
+	bool bSortingHasChanged = HasInProgressAsyncOperation(EAsyncOperationType::Sorting);
+
+	if (!bSortingHasChanged &&
+		CurrentSorter.Get() != CurrentAsyncOpSorter)
 	{
-		SortTreeNodes(CurrentSorter.Get(), ColumnSortMode);
+		bSortingHasChanged = true;
 	}
+
+	if (!bSortingHasChanged &&
+		CurrentSorter.IsValid() &&
+		ColumnSortMode != CurrentAsyncOpColumnSortMode)
+	{
+		bSortingHasChanged = true;
+	}
+
+	if (bSortingHasChanged)
+	{
+		ScheduleSortingAsyncOperation();
+		return true;
+	}
+
+	return false;
 }
 
 ////////////////////////////////////////////////////////////////////////////////////////////////////
 
+void STableTreeView::ScheduleSortingAsyncOperation()
+{
+	OnPreAsyncUpdate();
+
+	FGraphEventRef CompletedEvent = StartSortingTask();
+	InProgressAsyncOperationEvent = StartHierarchyFilteringTask(CompletedEvent);
+}
+
+////////////////////////////////////////////////////////////////////////////////////////////////////
+
+FGraphEventRef STableTreeView::StartSortingTask(FGraphEventRef Prerequisite)
+{
+	if (!CurrentSorter.IsValid())
+	{
+		return Prerequisite;
+	}
+
+	AddInProgressAsyncOperation(EAsyncOperationType::Sorting);
+
+	CurrentAsyncOpSorter = CurrentSorter.Get();
+	CurrentAsyncOpColumnSortMode = ColumnSortMode;
+
+	FGraphEventArray Prerequisites;
+	if (Prerequisite.IsValid())
+	{
+		Prerequisites.Add(Prerequisite);
+	}
+	else
+	{
+		Prerequisites.Add(DispatchEvent);
+	}
+
+	return TGraphTask<FTableTreeViewSortingAsyncTask>::CreateTask(&Prerequisites).ConstructAndDispatchWhenReady(this, CurrentAsyncOpSorter, CurrentAsyncOpColumnSortMode);
+}
+
+////////////////////////////////////////////////////////////////////////////////////////////////////
+
+// Note: This function is called from a task thread!
 void STableTreeView::SortTreeNodes(ITableCellValueSorter* InSorter, EColumnSortMode::Type InColumnSortMode)
 {
 	FStopwatch Stopwatch;
@@ -2418,14 +2699,8 @@
 
 void STableTreeView::SortTreeNodesRec(FTableTreeNode& GroupNode, const ITableCellValueSorter& Sorter, EColumnSortMode::Type InColumnSortMode)
 {
-	if (InColumnSortMode == EColumnSortMode::Type::Descending)
-	{
-		GroupNode.SortChildrenDescending(Sorter);
-	}
-	else // if (ColumnSortMode == EColumnSortMode::Type::Ascending)
-	{
-		GroupNode.SortChildrenAscending(Sorter);
-	}
+	const ESortMode SortMode = (InColumnSortMode == EColumnSortMode::Type::Descending) ? ESortMode::Descending : ESortMode::Ascending;
+	GroupNode.SortChildren(Sorter, SortMode);
 
 	for (FBaseTreeNodePtr ChildPtr : GroupNode.GetChildren())
 	{
@@ -2804,28 +3079,22 @@
 
 FTableTreeNodePtr STableTreeView::GetNodeByTableRowIndex(int32 RowIndex) const
 {
-	return (RowIndex >= 0 && RowIndex < TableTreeNodes.Num()) ? TableTreeNodes[RowIndex] : nullptr;
+	return (RowIndex >= 0 && RowIndex < TableRowNodes.Num()) ? TableRowNodes[RowIndex] : nullptr;
 }
 
 ////////////////////////////////////////////////////////////////////////////////////////////////////
 
 void STableTreeView::SelectNodeByTableRowIndex(int32 RowIndex)
 {
-	if (RowIndex >= 0 && RowIndex < TableTreeNodes.Num())
+	FTableTreeNodePtr NodePtr = GetNodeByTableRowIndex(RowIndex);
+	if (NodePtr.IsValid())
 	{
-		FTableTreeNodePtr NodePtr = TableTreeNodes[RowIndex];
-		ensure(NodePtr);
-
-		if (!NodePtr)
+		// Expand all parent nodes.
+		FBaseTreeNodePtr ParentNode = NodePtr->GetParent();
+		while (ParentNode.IsValid())
 		{
-			return;
-		}
-
-		FBaseTreeNodePtr GroupNode = NodePtr->GetGroupPtr().Pin();
-		while (GroupNode.IsValid())
-		{
-			TreeView->SetItemExpansion(StaticCastSharedPtr<FTableTreeNode>(GroupNode), true);
-			GroupNode = GroupNode->GetGroupPtr().Pin();
+			TreeView->SetItemExpansion(StaticCastSharedPtr<FTableTreeNode>(ParentNode), true);
+			ParentNode = ParentNode->GetParent();
 
 		}
 
@@ -2871,7 +3140,7 @@
 		TreeView->ClearExpandedItems();
 
 		// Grouping can result in old group nodes no longer existing in the tree, so we don't keep the expanded list.
-		if (!HasInProgressAsyncOperation(EAsyncOperationType::GroupingOp))
+		if (!HasInProgressAsyncOperation(EAsyncOperationType::Grouping))
 		{
 			for (auto It = ExpandedNodes.CreateConstIterator(); It; ++It)
 			{
@@ -2929,7 +3198,7 @@
 
 void STableTreeView::CountNumNodesPerDepthRec(FBaseTreeNode* InRoot, TArray<int32>& InOutNumNodesPerDepth, int32 InDepth, int32 InMaxDepth, int32 InMaxNodes) const
 {
-	InOutNumNodesPerDepth[InDepth] += InRoot->GetChildren().Num();
+	InOutNumNodesPerDepth[InDepth] += InRoot->GetChildrenCount();
 
 	if (InDepth < InMaxDepth && InOutNumNodesPerDepth[InDepth] < InMaxNodes)
 	{
@@ -2964,80 +3233,6 @@
 
 ////////////////////////////////////////////////////////////////////////////////////////////////////
 
-FGraphEventRef STableTreeView::StartSortTreeNodesTask(FGraphEventRef Prerequisite)
-{
-	if (!CurrentSorter.IsValid())
-	{
-		return Prerequisite;
-	}
-
-	AddInProgressAsyncOperation(EAsyncOperationType::SortingOp);
-
-	CurrentAsyncOpSorter = CurrentSorter.Get();
-	CurrentAsyncOpColumnSortMode = ColumnSortMode;
-
-	FGraphEventArray Prerequisites;
-	if (Prerequisite.IsValid())
-	{
-		Prerequisites.Add(Prerequisite);
-	}
-	else
-	{
-		Prerequisites.Add(DispatchEvent);
-	}
-
-	return TGraphTask<FTableTreeViewSortAsyncTask>::CreateTask(&Prerequisites).ConstructAndDispatchWhenReady(this, CurrentAsyncOpSorter, CurrentAsyncOpColumnSortMode);
-}
-
-////////////////////////////////////////////////////////////////////////////////////////////////////
-
-FGraphEventRef STableTreeView::StartCreateGroupsTask(FGraphEventRef Prerequisite)
-{
-	AddInProgressAsyncOperation(EAsyncOperationType::GroupingOp);
-
-	CurrentAsyncOpGroupings.Empty();
-	CurrentAsyncOpGroupings.Insert(CurrentGroupings, 0);
-
-	FGraphEventArray Prerequisites;
-	if (Prerequisite.IsValid())
-	{
-		Prerequisites.Add(Prerequisite);
-	}
-	else
-	{
-		Prerequisites.Add(DispatchEvent);
-	}
-
-	return TGraphTask<FTableTreeViewGroupAsyncTask>::CreateTask(&Prerequisites).ConstructAndDispatchWhenReady(this, &CurrentAsyncOpGroupings);
-}
-
-////////////////////////////////////////////////////////////////////////////////////////////////////
-
-FGraphEventRef STableTreeView::StartApplyFiltersTask(FGraphEventRef Prerequisite)
-{
-	AddInProgressAsyncOperation(EAsyncOperationType::FilteringOp);
-
-	CurrentAsyncOpTextFilter->SetRawFilterText(TextFilter->GetRawFilterText());
-	if (FilterConfigurator.IsValid() && CurrentAsyncOpFilterConfigurator)
-	{
-		*CurrentAsyncOpFilterConfigurator = *FilterConfigurator;
-	}
-
-	FGraphEventArray Prerequisites;
-	if (Prerequisite.IsValid())
-	{
-		Prerequisites.Add(Prerequisite);
-	}
-	else
-	{
-		Prerequisites.Add(DispatchEvent);
-	}
-
-	return TGraphTask<FTableTreeViewFilterAsyncTask>::CreateTask(&Prerequisites).ConstructAndDispatchWhenReady(this);
-}
-
-////////////////////////////////////////////////////////////////////////////////////////////////////
-
 void STableTreeView::OnClose()
 {
 	if (bIsCloseScheduled)
@@ -3077,65 +3272,23 @@
 
 void STableTreeView::StartPendingAsyncOperations()
 {
-	// Check if grouping settings have changed. If they did, a full refresh (Grouping, Sorting and Filtering) is scheduled.
-	bool bGroupingsHaveChanged = HasInProgressAsyncOperation(EAsyncOperationType::GroupingOp);
-	bGroupingsHaveChanged |= CurrentGroupings.Num() != CurrentAsyncOpGroupings.Num();
-
-	if (!bGroupingsHaveChanged)
+	if (ScheduleNodeFilteringAsyncOperationIfNeeded())
 	{
-		for (int Index = 0; Index < CurrentGroupings.Num(); ++Index)
-		{
-			if (CurrentGroupings[Index] != CurrentAsyncOpGroupings[Index])
-			{
-				bGroupingsHaveChanged = true;
-				break;
-			}
-		}
-	}
-
-	if (bGroupingsHaveChanged)
-	{
-		OnPreAsyncUpdate();
-
-		FGraphEventRef CompletedEvent = StartCreateGroupsTask();
-		CompletedEvent = StartSortTreeNodesTask(CompletedEvent);
-		InProgressAsyncOperationEvent = StartApplyFiltersTask(CompletedEvent);
-
 		return;
 	}
 
-	// Check if sorting settings have changed. If they did, a Sorting and Filtering Refresh is scheduled.
-	bool bSortingHasChanged = HasInProgressAsyncOperation(EAsyncOperationType::SortingOp);
-	bSortingHasChanged |= (CurrentSorter.IsValid() && CurrentAsyncOpSorter == nullptr) || (!CurrentSorter.IsValid() && CurrentAsyncOpSorter != nullptr);
-	if (!bSortingHasChanged && CurrentSorter.IsValid())
+	if (ScheduleGroupingAsyncOperationIfNeeded())
 	{
-		bSortingHasChanged = CurrentSorter.Get() != CurrentAsyncOpSorter || ColumnSortMode != CurrentAsyncOpColumnSortMode;
-	}
-
-	if (bSortingHasChanged)
-	{
-		OnPreAsyncUpdate();
-
-		FGraphEventRef CompletedEvent = StartSortTreeNodesTask();
-		InProgressAsyncOperationEvent = StartApplyFiltersTask(CompletedEvent);
-
 		return;
 	}
 
-	// Check if the text filter has changed. If it has, schedule a new Filtering Refresh.
-	bool bFiltersHaveChanged = HasInProgressAsyncOperation(EAsyncOperationType::FilteringOp);
-	bFiltersHaveChanged |= TextFilter->GetRawFilterText().CompareTo(CurrentAsyncOpTextFilter->GetRawFilterText()) != 0;
-	if (FilterConfigurator.IsValid() && CurrentAsyncOpFilterConfigurator && !bFiltersHaveChanged)
+	if (ScheduleSortingAsyncOperationIfNeeded())
 	{
-		bFiltersHaveChanged |= *FilterConfigurator != *CurrentAsyncOpFilterConfigurator;
+		return;
 	}
 
-	if (bFiltersHaveChanged)
+	if (ScheduleHierarchyFilteringAsyncOperationIfNeeded())
 	{
-		OnPreAsyncUpdate();
-
-		InProgressAsyncOperationEvent = StartApplyFiltersTask();
-
 		return;
 	}
 }
@@ -3152,152 +3305,6 @@
 
 ////////////////////////////////////////////////////////////////////////////////////////////////////
 
-FReply STableTreeView::OnAdvancedFiltersClicked()
-{
-	if (!FilterConfigurator.IsValid())
-	{
-		FilterConfigurator = MakeShared<FFilterConfigurator>();
-		TSharedPtr<TArray<TSharedPtr<struct FFilter>>>& AvailableFilters = FilterConfigurator->GetAvailableFilters();
-
-		for (const TSharedRef<FTableColumn>& Column : Table->GetColumns())
-		{
-			if (!Column->CanBeFiltered())
-			{
-				continue;
-			}
-
-			switch (Column->GetDataType())
-			{
-				case ETableCellDataType::Int64:
-				{
-					AvailableFilters->Add(MakeShared<FFilter>(Column->GetIndex(), Column->GetTitleName(), Column->GetDescription(), EFilterDataType::Int64, FFilterService::Get()->GetIntegerOperators()));
-					AvailableFilters->Last()->Converter = Column->GetValueConverter();
-					Context.AddFilterData<int64>(Column->GetIndex(), 0);
-					break;
-				}
-				case ETableCellDataType::Double:
-				{
-					AvailableFilters->Add(MakeShared<FFilter>(Column->GetIndex(), Column->GetTitleName(), Column->GetDescription(), EFilterDataType::Double, FFilterService::Get()->GetDoubleOperators()));
-					AvailableFilters->Last()->Converter = Column->GetValueConverter();
-					Context.AddFilterData<double>(Column->GetIndex(), 0.0);
-					break;
-				}
-				case ETableCellDataType::CString:
-				case ETableCellDataType::Text:
-				case ETableCellDataType::Custom:
-				{
-					if (!Column->IsHierarchy())
-					{
-						AvailableFilters->Add(MakeShared<FFilter>(Column->GetIndex(), Column->GetTitleName(), Column->GetDescription(), EFilterDataType::String, FFilterService::Get()->GetStringOperators()));
-						AvailableFilters->Last()->Converter = Column->GetValueConverter();
-						Context.AddFilterData<FString>(Column->GetIndex(), FString());
-					}
-					break;
-				}
-			}
-		}
-
-		AddCustomAdvancedFilters();
-
-		CurrentAsyncOpFilterConfigurator = new FFilterConfigurator(*FilterConfigurator);
-		OnFilterChangesCommittedHandle = FilterConfigurator->GetOnChangesCommittedEvent().AddSP(this, &STableTreeView::OnAdvancedFiltersChangesCommitted);
-	}
-
-	FFilterService::Get()->CreateFilterConfiguratorWidget(FilterConfigurator);
-
-	return FReply::Handled();
-}
-
-////////////////////////////////////////////////////////////////////////////////////////////////////
-
-bool STableTreeView::ApplyAdvancedFilters(const FTableTreeNodePtr& NodePtr)
-{
-	FFilterConfigurator* FilterConfiguratorToUse = CurrentAsyncOpFilterConfigurator;
-
-	if (FilterConfiguratorToUse == nullptr)
-	{
-		return true;
-	}
-
-	if (FilterConfiguratorToUse->GetRootNode()->GetChildren().Num() == 0)
-	{
-		return true;
-	}
-
-	for (const TSharedRef<FTableColumn>& Column : Table->GetColumns())
-	{
-		if (!Column->CanBeFiltered())
-		{
-			continue;
-		}
-
-		switch (Column->GetDataType())
-		{
-			case ETableCellDataType::Int64:
-			{
-				Context.SetFilterData<int64>(Column->GetIndex(), Column->GetValue(*NodePtr)->AsInt64());
-				break;
-			}
-			case ETableCellDataType::Double:
-			{
-				Context.SetFilterData<double>(Column->GetIndex(), Column->GetValue(*NodePtr)->AsDouble());
-				break;
-			}
-			case ETableCellDataType::CString:
-			case ETableCellDataType::Text:
-			case ETableCellDataType::Custom:
-			{
-				if (!Column->IsHierarchy())
-				{
-					// Converting string is heavy, check if the key is used for filtering before doing any conversion
-					if (FilterConfiguratorToUse->IsKeyUsed(Column->GetIndex()))
-					{
-						Context.SetFilterData<FString>(Column->GetIndex(), Column->GetValue(*NodePtr)->AsString());
-					}
-				}
-				break;
-			}
-		}
-	}
-
-	return ApplyCustomAdvancedFilters(NodePtr) && FilterConfiguratorToUse->ApplyFilters(Context);
-}
-
-////////////////////////////////////////////////////////////////////////////////////////////////////
-
-void STableTreeView::OnAdvancedFiltersChangesCommitted()
-{
-	OnFilteringChanged();
-}
-
-////////////////////////////////////////////////////////////////////////////////////////////////////
-
-bool STableTreeView::AdvancedFilters_ShouldBeEnabled() const
-{
-	return TextFilter->GetRawFilterText().IsEmpty();
-}
-
-////////////////////////////////////////////////////////////////////////////////////////////////////
-
-FText STableTreeView::AdvancedFilters_GetTooltipText() const
-{
-	if (AdvancedFilters_ShouldBeEnabled())
-	{
-		return LOCTEXT("AdvancedFiltersBtn_ToolTip", "Opens the filter configurator window.");
-	}
-
-	return LOCTEXT("AdvancedFiltersBtn_Disabled_ToolTip", "Advanced filters cannot be added when filters are already applied using the search box.");
-}
-
-////////////////////////////////////////////////////////////////////////////////////////////////////
-
-bool STableTreeView::FilterConfigurator_HasFilters() const
-{
-	return FilterConfigurator.IsValid() && FilterConfigurator->GetRootNode()->GetChildren().Num() > 0;
-}
-
-////////////////////////////////////////////////////////////////////////////////////////////////////
-
 FReply STableTreeView::OnKeyDown(const FGeometry& MyGeometry, const FKeyEvent& InKeyEvent)
 {
 	return CommandList->ProcessCommandBindings(InKeyEvent) == true ? FReply::Handled() : FReply::Unhandled();
@@ -3644,6 +3651,25 @@
 
 ////////////////////////////////////////////////////////////////////////////////////////////////////
 
+void STableTreeView::UpdateBannerText()
+{
+	TreeViewBannerText = FText::GetEmpty();
+	if (TableRowNodes.Num() == 0)
+	{
+		TreeViewBannerText = LOCTEXT("EmptyTable", "This table is empty.");
+	}
+	else if (FilteredNodesPtr->Num() == 0)
+	{
+		TreeViewBannerText = LOCTEXT("AllNodesFilteredOut", "All tree nodes are filtered out. Check the filter configurator.");
+	}
+	else if (Root->GetFilteredChildrenCount() == 0)
+	{
+		TreeViewBannerText = LOCTEXT("HierarchyFilteringZeroResults", "No tree node is matching the current text search.");
+	}
+}
+
+////////////////////////////////////////////////////////////////////////////////////////////////////
+
 FReply STableTreeView::OnApplyViewPreset(const ITableTreeViewPreset* InPreset)
 {
 	ApplyViewPreset(*InPreset);
diff --git a/Engine/Plugins/Editor/AssetManagerEditor/Source/AssetManagerEditor/Private/Insights/Table/Widgets/STableTreeView.h b/Engine/Plugins/Editor/AssetManagerEditor/Source/AssetManagerEditor/Private/Insights/Table/Widgets/STableTreeView.h
--- a/Engine/Plugins/Editor/AssetManagerEditor/Source/AssetManagerEditor/Private/Insights/Table/Widgets/STableTreeView.h
+++ b/Engine/Plugins/Editor/AssetManagerEditor/Source/AssetManagerEditor/Private/Insights/Table/Widgets/STableTreeView.h
@@ -53,9 +53,11 @@ typedef TTextFilter<const FTableTreeNodePtr&> FTableTreeNodeTextFilter;
 
 enum class EAsyncOperationType : uint32
 {
-	FilteringOp = 1,
-	SortingOp = 1 << 1,
-	GroupingOp = 1 << 2,
+	NodeFiltering       = 1 << 0,
+	Grouping            = 1 << 1,
+	Aggregation         = 1 << 2,
+	Sorting             = 1 << 3,
+	HierarchyFiltering  = 1 << 4,
 };
 
 ENUM_CLASS_FLAGS(EAsyncOperationType);
@@ -104,9 +106,10 @@ struct FTableTaskInfo
  */
 class STableTreeView : public SCompoundWidget, public IAsyncOperationStatusProvider
 {
-	friend class FTableTreeViewFilterAsyncTask;
-	friend class FTableTreeViewSortAsyncTask;
-	friend class FTableTreeViewGroupAsyncTask;
+	friend class FTableTreeViewNodeFilteringAsyncTask;
+	friend class FTableTreeViewSortingAsyncTask;
+	friend class FTableTreeViewGroupingAsyncTask;
+	friend class FTableTreeViewHierarchyFilteringAsyncTask;
 	friend class FTableTreeViewAsyncCompleteTask;
 	friend class FSearchForItemToSelectTask;
 	friend class FSelectNodeByTableRowIndexTask;
@@ -179,8 +182,6 @@ protected:
 	void InitCommandList();
 
 	virtual void ConstructWidget(TSharedPtr<FTable> InTablePtr);
-	virtual TSharedRef<SWidget> ConstructSearchBox();
-	virtual TSharedRef<SWidget> ConstructAdvancedFiltersButton();
 	virtual TSharedRef<SWidget> ConstructHierarchyBreadcrumbTrail();
 	virtual TSharedPtr<SWidget> ConstructToolbar() { return nullptr; }
 	virtual TSharedPtr<SWidget> ConstructFooter() { return nullptr; }
@@ -189,15 +190,6 @@ protected:
 
 	void UpdateTree();
 
-	/**
-	 * Populates OutSearchStrings with the strings that should be used in searching.
-	 *
-	 * @param GroupOrStatNodePtr - the group and stat node to get a text description from.
-	 * @param OutSearchStrings   - an array of strings to use in searching.
-	 *
-	 */
-	static void HandleItemToStringArray(const FTableTreeNodePtr& GroupOrStatNodePtr, TArray<FString>& OutSearchStrings);
-
 	////////////////////////////////////////////////////////////////////////////////////////////////////
 	// Tree View - Context Menu
 
@@ -266,33 +258,63 @@ protected:
 	FName TableRow_GetHighlightedNodeName() const;
 
 	////////////////////////////////////////////////////////////////////////////////////////////////////
-	// Filtering
+	// Node Filtering (TableRowNodes --> FilteredNodes)
+
+	void InitNodeFiltering();
+
+	void OnNodeFilteringChanged();
+	bool ScheduleNodeFilteringAsyncOperationIfNeeded();
+	void ScheduleNodeFilteringAsyncOperation();
+	FGraphEventRef StartNodeFilteringTask(FGraphEventRef Prerequisite = nullptr);
+
+	void ApplyNodeFiltering();
+
+	virtual bool FilterNode(const FFilterConfigurator& InFilterConfigurator, const FTableTreeNode& InNode) const;
+
+	virtual void InitFilterConfigurator(FFilterConfigurator& InOutFilterConfigurator);
+	virtual void UpdateFilterContext(const FFilterConfigurator& InFilterConfigurator, const FTableTreeNode& InNode) const;
+
+	virtual TSharedRef<SWidget> ConstructFilterConfiguratorButton();
+	FReply FilterConfigurator_OnClicked();
+
+	////////////////////////////////////////////////////////////////////////////////////////////////////
+	// Hierarchy Filtering (Root->Children hierarchy --> Root->FilteredChildren hierarchy)
+
+	void InitHierarchyFiltering();
+
+	/**
+	 * Populates OutSearchStrings with the strings that should be used in searching.
+	 *
+	 * @param GroupOrStatNodePtr - the group and stat node to get a text description from.
+	 * @param OutSearchStrings   - an array of strings to use in searching.
+	 */
+	static void HandleItemToStringArray(const FTableTreeNodePtr& GroupOrStatNodePtr, TArray<FString>& OutSearchStrings);
 
-	void OnFilteringChanged();
+	void OnHierarchyFilteringChanged();
+	bool ScheduleHierarchyFilteringAsyncOperationIfNeeded();
+	void ScheduleHierarchyFilteringAsyncOperation();
+	FGraphEventRef StartHierarchyFilteringTask(FGraphEventRef Prerequisite = nullptr);
 
-	/** Populates the group and stat tree with items based on the current data. */
-	void ApplyFiltering();
-	void ApplyFiltering(FTableTreeNodePtr NodePtr);
-	bool ApplyAdvancedFiltersForNode(FTableTreeNodePtr NodePtr);
-	bool ApplyHierarchicalFilterForNode(FTableTreeNodePtr NodePtr, bool bFilterIsEmpty);
+	void ApplyHierarchyFiltering();
+	void ApplyEmptyHierarchyFilteringRec(FTableTreeNodePtr NodePtr);
+	bool ApplyHierarchyFilteringRec(FTableTreeNodePtr NodePtr);
 
 	/** Set all the nodes belonging to a subtree as visible. Returns true if the caller node should be expanded. */
 	bool MakeSubtreeVisible(FTableTreeNodePtr NodePtr, bool bFilterIsEmpty);
 
-	bool SearchBox_IsEnabled() const;
+	virtual TSharedRef<SWidget> ConstructSearchBox();
 	void SearchBox_OnTextChanged(const FText& InFilterText);
 
-	FText SearchBox_GetTooltipText() const;
-
 	////////////////////////////////////////////////////////////////////////////////////////////////////
-	// Grouping
+	// Grouping (FilteredNodes --> Root->Children hierarchy)
 
 	void CreateGroupings();
 	virtual void InternalCreateGroupings();
 
 	void OnGroupingChanged();
-	void ApplyGrouping();
-
+	bool ScheduleGroupingAsyncOperationIfNeeded();
+	void ScheduleGroupingAsyncOperation();
+	FGraphEventRef StartGroupingTask(FGraphEventRef Prerequisite = nullptr);
 	void CreateGroups(const TArray<TSharedPtr<FTreeNodeGrouping>>& Groupings);
 	void GroupNodesRec(const TArray<FTableTreeNodePtr>& Nodes, FTableTreeNode& ParentGroup, int32 GroupingDepth, const TArray<TSharedPtr<FTreeNodeGrouping>>& Groupings);
 
@@ -316,7 +338,7 @@ protected:
 	bool GroupingCrumbMenu_Add_CanExecute(const TSharedPtr<FTreeNodeGrouping> Grouping, const TSharedPtr<FTreeNodeGrouping> AfterGrouping) const;
 
 	////////////////////////////////////////////////////////////////////////////////////////////////////
-	// Aggregation
+	// Aggregation (Root->Children hierarchy)
 
 	static void UpdateCStringSameValueAggregationSingleNode(const FTableColumn& InColumn, FTableTreeNode& GroupNode);
 	static void UpdateCStringSameValueAggregationRec(const FTableColumn& InColumn, FTableTreeNode& GroupNode);
@@ -331,17 +353,18 @@ protected:
 	void UpdateAggregatedValuesRec(FTableTreeNode& GroupNode);
 
 	////////////////////////////////////////////////////////////////////////////////////////////////////
-	// Sorting
+	// Sorting (Root->Children hierarchy)
 
 	static const EColumnSortMode::Type GetDefaultColumnSortMode();
 	static const FName GetDefaultColumnBeingSorted();
 
-	void OnSortingChanged();
 	void CreateSortings();
-
 	void UpdateCurrentSortingByColumn();
 
-	void ApplySorting();
+	void OnSortingChanged();
+	bool ScheduleSortingAsyncOperationIfNeeded();
+	void ScheduleSortingAsyncOperation();
+	FGraphEventRef StartSortingTask(FGraphEventRef Prerequisite = nullptr);
 	void SortTreeNodes(ITableCellValueSorter* InSorter, EColumnSortMode::Type InColumnSortMode);
 	void SortTreeNodesRec(FTableTreeNode& GroupNode, const ITableCellValueSorter& Sorter, EColumnSortMode::Type InColumnSortMode);
 
@@ -397,14 +420,12 @@ protected:
 	bool ContextMenu_HideAllColumns_CanExecute() const;
 	void ContextMenu_HideAllColumns_Execute();
 
+	////////////////////////////////////////////////////////////////////////////////////////////////////
 	//Async
+
 	virtual void OnPreAsyncUpdate();
 	virtual void OnPostAsyncUpdate();
 
-	FGraphEventRef StartSortTreeNodesTask(FGraphEventRef Prerequisite = nullptr);
-	FGraphEventRef StartCreateGroupsTask(FGraphEventRef Prerequisite = nullptr);
-	FGraphEventRef StartApplyFiltersTask(FGraphEventRef Prerequisite = nullptr);
-
 	void AddInProgressAsyncOperation(EAsyncOperationType InType) { EnumAddFlags(InProgressAsyncOperations, InType); }
 	bool HasInProgressAsyncOperation(EAsyncOperationType InType) const { return EnumHasAnyFlags(InProgressAsyncOperations, InType); }
 	void ClearInProgressAsyncOperations() { InProgressAsyncOperations = static_cast<EAsyncOperationType>(0); }
@@ -413,15 +434,6 @@ protected:
 
 	void CancelCurrentAsyncOp();
 
-	FReply OnAdvancedFiltersClicked();
-	bool AdvancedFilters_ShouldBeEnabled() const;
-	FText AdvancedFilters_GetTooltipText() const;
-	bool FilterConfigurator_HasFilters() const;
-	void OnAdvancedFiltersChangesCommitted();
-	bool ApplyAdvancedFilters(const FTableTreeNodePtr& NodePtr);
-	bool virtual ApplyCustomAdvancedFilters(const FTableTreeNodePtr& NodePtr) { return true; };
-	virtual void AddCustomAdvancedFilters() {}
-
 	////////////////////////////////////////////////////////////////////////////////////////////////////
 
 	void SetExpandValueForChildGroups(FBaseTreeNode* InRoot, int32 InMaxExpandedNodes, int32 MaxDepthToExpand, bool InValue);
@@ -434,6 +446,10 @@ protected:
 	void ExportToFileRec(const FBaseTreeNodePtr& InGroupNode, TArray<FBaseTreeNodePtr>& InNodes, bool bInExportCollapsed, bool InExportLeafs, WriteToFileCallback Callback);
 
 	FText GetTreeViewBannerText() const { return TreeViewBannerText; }
+	virtual void UpdateBannerText();
+
+	////////////////////////////////////////////////////////////////////////////////////////////////////
+	// Presets
 
 	virtual void InitAvailableViewPresets() {};
 	const TArray<TSharedRef<ITableTreeViewPreset>>* GetAvailableViewPresets() const { return &AvailableViewPresets; }
@@ -445,6 +461,8 @@ protected:
 	FText ViewPreset_GetSelectedText() const;
 	FText ViewPreset_GetSelectedToolTipText() const;
 
+	////////////////////////////////////////////////////////////////////////////////////////////////////
+
 	virtual void SearchForItem(TSharedPtr<FTableTaskCancellationToken> CancellationToken) {};
 
 	// Table data tasks should be tasks that operate read only operations on the data from the Table
@@ -456,7 +474,6 @@ protected:
 		Info->CancellationToken = MakeShared<FTableTaskCancellationToken>();
 		Info->Event = TGraphTask<T>::CreateTask().ConstructAndDispatchWhenReady(Info->CancellationToken, Forward<TArgs>(Args)...);
 		DataTaskInfos.Add(Info);
-
 		return Info;
 	}
 
@@ -500,10 +517,16 @@ protected:
 	/** The root node of the tree. */
 	FTableTreeNodePtr Root;
 
-	/** Table (row) nodes. Each node corresponds to a table row. Index in this array corresponds to RowIndex in source table. */
-	TArray<FTableTreeNodePtr> TableTreeNodes;
+	/** Table row nodes. Each node corresponds to a table row. Index in this array corresponds to RowIndex in source table. */
+	TArray<FTableTreeNodePtr> TableRowNodes;
+
+	/** Filtered table nodes. These are the nodes from the TableRowNodes array, after applying the node filtering. */
+	TArray<FTableTreeNodePtr> FilteredNodes;
+
+	/** A pointer to the filtered table nodes. If node filtering is empty, this points directly to TableRowNodes. */
+	TArray<FTableTreeNodePtr>* FilteredNodesPtr = nullptr;
 
-	/** A filtered array of group and nodes to be displayed in the tree widget. */
+	/** A filtered array of group nodes to be displayed in the tree widget. */
 	TArray<FTableTreeNodePtr> FilteredGroupNodes;
 
 	/** Currently expanded group nodes. */
@@ -514,23 +537,31 @@ protected:
 	static constexpr int32 MaxNodesToExpand = 1000000;
 	static constexpr int32 MaxDepthToExpand = 100;
 
-	EInvalidateWidgetReason LazyInvalidateWidgetReason = EInvalidateWidgetReason::None;
-
 	//////////////////////////////////////////////////
-	// Search box and filters
+	// Search box & the hierarchy filtering
 
 	/** The search box widget used to filter items displayed in the stats and groups tree. */
 	TSharedPtr<SSearchBox> SearchBox;
 
+	/** The filter collection. */
+	TSharedPtr<FTableTreeNodeFilterCollection> Filters;
+
 	/** The text based filter. */
 	TSharedPtr<FTableTreeNodeTextFilter> TextFilter;
 
-	/** The filter collection. */
-	TSharedPtr<FTableTreeNodeFilterCollection> Filters;
+	/** The text based filter actually used in the Hierarchy Filtering async task. */
+	TSharedPtr<FTableTreeNodeTextFilter> CurrentAsyncOpTextFilter;
+
+	//////////////////////////////////////////////////
+	// Node filtering
 
 	TSharedPtr<FFilterConfigurator> FilterConfigurator;
 	FDelegateHandle OnFilterChangesCommittedHandle;
-	FFilterContext Context;
+
+	/** The filter configurator actually used in the Hierarchy Filtering async task. */
+	FFilterConfigurator* CurrentAsyncOpFilterConfigurator = nullptr;
+
+	mutable FFilterContext FilterContext;
 
 	//////////////////////////////////////////////////
 	// Grouping
@@ -542,6 +573,9 @@ protected:
 
 	TSharedPtr<SBreadcrumbTrail<TSharedPtr<FTreeNodeGrouping>>> GroupingBreadcrumbTrail;
 
+	/** The groupings actually used in the Grouping async task. */
+	TArray<TSharedPtr<FTreeNodeGrouping>> CurrentAsyncOpGroupings;
+
 	//////////////////////////////////////////////////
 	// Sorting
 
@@ -557,6 +591,12 @@ protected:
 	/** How we sort the nodes? Ascending or Descending. */
 	EColumnSortMode::Type ColumnSortMode = EColumnSortMode::None;
 
+	/** The sorter actually used in the Sorting async task. */
+	ITableCellValueSorter* CurrentAsyncOpSorter = nullptr;
+
+	/** The sort mode actually used in the Sorting async task. */
+	EColumnSortMode::Type CurrentAsyncOpColumnSortMode;
+
 	//////////////////////////////////////////////////
 	// Async Operations
 
@@ -569,12 +609,6 @@ protected:
 	EAsyncOperationType InProgressAsyncOperations = static_cast<EAsyncOperationType>(0);
 	TSharedPtr<class SAsyncOperationStatus> AsyncOperationStatus;
 	FStopwatch AsyncUpdateStopwatch;
-
-	TArray<TSharedPtr<FTreeNodeGrouping>> CurrentAsyncOpGroupings;
-	ITableCellValueSorter* CurrentAsyncOpSorter = nullptr;
-	EColumnSortMode::Type CurrentAsyncOpColumnSortMode;
-	TSharedPtr<FTableTreeNodeTextFilter> CurrentAsyncOpTextFilter;
-	FFilterConfigurator* CurrentAsyncOpFilterConfigurator = nullptr;
 	FAsyncOperationProgress AsyncOperationProgress;
 	FGraphEventRef DispatchEvent;
 	TArray<TSharedPtr<FTableTaskInfo>> DataTaskInfos;
@@ -597,15 +631,41 @@ protected:
 
 ////////////////////////////////////////////////////////////////////////////////////////////////////
 
-class FTableTreeViewFilterAsyncTask
+class FTableTreeViewNodeFilteringAsyncTask
+{
+public:
+	FTableTreeViewNodeFilteringAsyncTask(STableTreeView* InPtr)
+	{
+		TableTreeViewPtr = InPtr;
+	}
+
+	FORCEINLINE TStatId GetStatId() const { RETURN_QUICK_DECLARE_CYCLE_STAT(FTableTreeViewNodeFilteringAsyncTask, STATGROUP_TaskGraphTasks); }
+	ENamedThreads::Type GetDesiredThread() { return ENamedThreads::Type::AnyThread; }
+	static ESubsequentsMode::Type GetSubsequentsMode() { return ESubsequentsMode::TrackSubsequents; }
+
+	void DoTask(ENamedThreads::Type CurrentThread, const FGraphEventRef& MyCompletionGraphEvent)
+	{
+		if (TableTreeViewPtr)
+		{
+			TableTreeViewPtr->ApplyNodeFiltering();
+		}
+	}
+
+private:
+	STableTreeView* TableTreeViewPtr = nullptr;
+};
+
+////////////////////////////////////////////////////////////////////////////////////////////////////
+
+class FTableTreeViewHierarchyFilteringAsyncTask
 {
 public:
-	FTableTreeViewFilterAsyncTask(STableTreeView* InPtr)
+	FTableTreeViewHierarchyFilteringAsyncTask(STableTreeView* InPtr)
 	{
 		TableTreeViewPtr = InPtr;
 	}
 
-	FORCEINLINE TStatId GetStatId() const { RETURN_QUICK_DECLARE_CYCLE_STAT(FTableTreeViewFilterAsyncTask, STATGROUP_TaskGraphTasks); }
+	FORCEINLINE TStatId GetStatId() const { RETURN_QUICK_DECLARE_CYCLE_STAT(FTableTreeViewHierarchyFilteringAsyncTask, STATGROUP_TaskGraphTasks); }
 	ENamedThreads::Type GetDesiredThread() { return ENamedThreads::Type::AnyThread; }
 	static ESubsequentsMode::Type GetSubsequentsMode() { return ESubsequentsMode::TrackSubsequents; }
 
@@ -613,7 +673,7 @@ public:
 	{
 		if (TableTreeViewPtr)
 		{
-			TableTreeViewPtr->ApplyFiltering();
+			TableTreeViewPtr->ApplyHierarchyFiltering();
 		}
 	}
 
@@ -623,17 +683,17 @@ private:
 
 ////////////////////////////////////////////////////////////////////////////////////////////////////
 
-class FTableTreeViewSortAsyncTask
+class FTableTreeViewSortingAsyncTask
 {
 public:
-	FTableTreeViewSortAsyncTask(STableTreeView* InPtr, ITableCellValueSorter* InSorter, EColumnSortMode::Type InColumnSortMode)
+	FTableTreeViewSortingAsyncTask(STableTreeView* InPtr, ITableCellValueSorter* InSorter, EColumnSortMode::Type InColumnSortMode)
 	{
 		TableTreeViewPtr = InPtr;
 		Sorter = InSorter;
 		ColumnSortMode = InColumnSortMode;
 	}
 
-	FORCEINLINE TStatId GetStatId() const { RETURN_QUICK_DECLARE_CYCLE_STAT(FTableTreeViewSortAsyncTask, STATGROUP_TaskGraphTasks); }
+	FORCEINLINE TStatId GetStatId() const { RETURN_QUICK_DECLARE_CYCLE_STAT(FTableTreeViewSortingAsyncTask, STATGROUP_TaskGraphTasks); }
 	ENamedThreads::Type GetDesiredThread() { return ENamedThreads::Type::AnyThread; }
 	static ESubsequentsMode::Type GetSubsequentsMode() { return ESubsequentsMode::TrackSubsequents; }
 
@@ -653,16 +713,16 @@ private:
 
 ////////////////////////////////////////////////////////////////////////////////////////////////////
 
-class FTableTreeViewGroupAsyncTask
+class FTableTreeViewGroupingAsyncTask
 {
 public:
-	FTableTreeViewGroupAsyncTask(STableTreeView* InPtr, TArray<TSharedPtr<FTreeNodeGrouping>>* InGroupings)
+	FTableTreeViewGroupingAsyncTask(STableTreeView* InPtr, TArray<TSharedPtr<FTreeNodeGrouping>>* InGroupings)
 	{
 		TableTreeViewPtr = InPtr;
 		Groupings = InGroupings;
 	}
 
-	FORCEINLINE TStatId GetStatId() const { RETURN_QUICK_DECLARE_CYCLE_STAT(FTableTreeViewGroupAsyncTask, STATGROUP_TaskGraphTasks); }
+	FORCEINLINE TStatId GetStatId() const { RETURN_QUICK_DECLARE_CYCLE_STAT(FTableTreeViewGroupingAsyncTask, STATGROUP_TaskGraphTasks); }
 	ENamedThreads::Type GetDesiredThread() { return ENamedThreads::Type::AnyThread; }
 	static ESubsequentsMode::Type GetSubsequentsMode() { return ESubsequentsMode::TrackSubsequents; }
 
diff --git a/Engine/Plugins/Editor/AssetManagerEditor/Source/AssetManagerEditor/Private/Insights/Table/Widgets/STableTreeViewCell.cpp b/Engine/Plugins/Editor/AssetManagerEditor/Source/AssetManagerEditor/Private/Insights/Table/Widgets/STableTreeViewCell.cpp
--- a/Engine/Plugins/Editor/AssetManagerEditor/Source/AssetManagerEditor/Private/Insights/Table/Widgets/STableTreeViewCell.cpp
+++ b/Engine/Plugins/Editor/AssetManagerEditor/Source/AssetManagerEditor/Private/Insights/Table/Widgets/STableTreeViewCell.cpp
@@ -106,7 +106,6 @@ TSharedRef<SWidget> STableTreeViewCell::GenerateWidgetForNameColumn(const FArgum
 			.TextStyle(FInsightsStyle::Get(), TEXT("TreeTable.NameText"))
 			.ColorAndOpacity(this, &STableTreeViewCell::GetDisplayNameColorAndOpacity)
 			.ShadowColorAndOpacity(this, &STableTreeViewCell::GetShadowColorAndOpacity)
-			.ToolTip(RowToolTip)
 		]
 
 		// Name Suffix
@@ -122,7 +121,6 @@ TSharedRef<SWidget> STableTreeViewCell::GenerateWidgetForNameColumn(const FArgum
 			.TextStyle(FInsightsStyle::Get(), TEXT("TreeTable.NameText"))
 			.ColorAndOpacity(this, &STableTreeViewCell::GetExtraDisplayNameColorAndOpacity)
 			.ShadowColorAndOpacity(this, &STableTreeViewCell::GetShadowColorAndOpacity)
-			.ToolTip(RowToolTip)
 		]
 	;
 }
diff --git a/Engine/Plugins/Editor/AssetManagerEditor/Source/AssetManagerEditor/Private/Insights/Table/Widgets/STableTreeViewTooltip.cpp b/Engine/Plugins/Editor/AssetManagerEditor/Source/AssetManagerEditor/Private/Insights/Table/Widgets/STableTreeViewTooltip.cpp
--- a/Engine/Plugins/Editor/AssetManagerEditor/Source/AssetManagerEditor/Private/Insights/Table/Widgets/STableTreeViewTooltip.cpp
+++ b/Engine/Plugins/Editor/AssetManagerEditor/Source/AssetManagerEditor/Private/Insights/Table/Widgets/STableTreeViewTooltip.cpp
@@ -108,7 +108,7 @@ TSharedPtr<SToolTip> STableTreeViewTooltip::GetRowTooltip(const TSharedPtr<FTabl
 	TSharedPtr<SGridPanel> GridPanel;
 	TSharedPtr<SHorizontalBox> HBox;
 
-	const FText NodeTooltip = TreeNodePtr->GetTooltip();
+	const FText NodeTooltip = TreeNodePtr->GetTooltipText();
 	const EVisibility NodeTooltipVisibility = NodeTooltip.IsEmpty() ? EVisibility::Collapsed : EVisibility::Visible;
 
 	TSharedPtr<SToolTip> TableCellTooltip =
@@ -201,7 +201,7 @@ TSharedPtr<SToolTip> STableTreeViewTooltip::GetRowTooltip(const TSharedPtr<FTabl
 				[
 					SNew(STextBlock)
 					.Visibility(NodeTooltipVisibility)
-					.WrapTextAt(512.0f)
+					.WrapTextAt(1024.0f)
 					.WrappingPolicy(ETextWrappingPolicy::AllowPerCharacterWrapping)
 					.Text(NodeTooltip)
 					.TextStyle(GetTableTreeViewTooltipNormalTextStyle())
diff --git a/Engine/Plugins/Editor/AssetManagerEditor/Source/AssetManagerEditor/Private/SAssetAuditBrowser.cpp b/Engine/Plugins/Editor/AssetManagerEditor/Source/AssetManagerEditor/Private/SAssetAuditBrowser.cpp
--- a/Engine/Plugins/Editor/AssetManagerEditor/Source/AssetManagerEditor/Private/SAssetAuditBrowser.cpp
+++ b/Engine/Plugins/Editor/AssetManagerEditor/Source/AssetManagerEditor/Private/SAssetAuditBrowser.cpp
@@ -820,7 +820,6 @@ void SAssetAuditBrowser::TreeView_OnSelectionChanged(const TArray<TSharedPtr<UE:
 	TSharedPtr<FAssetTreeNode> SelectedAssetNode = AssetTableTreeView->GetSingleSelectedAssetNode();
 	if (SelectedAssetNode.IsValid())
 	{
-		check(!SelectedAssetNode->IsGroup());
 		check(SelectedAssetNode->IsValidAsset());
 		const FAssetTableRow& AssetTableRow = SelectedAssetNode->GetAssetChecked();
 		//...
diff --git a/Engine/Plugins/Editor/AssetManagerEditor/Source/AssetManagerEditor/Private/TreeView/AssetDependencyGrouping.cpp b/Engine/Plugins/Editor/AssetManagerEditor/Source/AssetManagerEditor/Private/TreeView/AssetDependencyGrouping.cpp
--- a/Engine/Plugins/Editor/AssetManagerEditor/Source/AssetManagerEditor/Private/TreeView/AssetDependencyGrouping.cpp
+++ b/Engine/Plugins/Editor/AssetManagerEditor/Source/AssetManagerEditor/Private/TreeView/AssetDependencyGrouping.cpp
@@ -40,8 +40,6 @@ void FAssetDependencyGrouping::GroupNodes(const TArray<UE::Insights::FTableTreeN
 	TSharedPtr<FAssetTable> AssetTable = StaticCastSharedPtr<FAssetTable>(InParentTable.Pin());
 	check(AssetTable.IsValid());
 
-	const FSlateBrush* AssetGroupNodeIconBrush = FBaseTreeNode::GetDefaultIcon(true);
-
 	for (FTableTreeNodePtr NodePtr : Nodes)
 	{
 		if (InAsyncOperationProgress.ShouldCancelAsyncOp())
@@ -51,7 +49,7 @@ void FAssetDependencyGrouping::GroupNodes(const TArray<UE::Insights::FTableTreeN
 
 		if (NodePtr->IsGroup() || !NodePtr->Is<FAssetTreeNode>())
 		{
-			ParentGroup.AddChildAndSetGroupPtr(NodePtr);
+			ParentGroup.AddChildAndSetParent(NodePtr);
 			continue;
 		}
 
@@ -91,28 +89,28 @@ void FAssetDependencyGrouping::GroupNodes(const TArray<UE::Insights::FTableTreeN
 		if (Asset.GetNumDependencies() > 0)
 		{
 			// Create a group for the asset node (self) + dependencies.
-			FTableTreeNodePtr AssetGroupPtr = MakeShared<FCustomTableTreeNode>(Asset.GetNodeName(), InParentTable, AssetNode.GetRowIndex(), AssetGroupNodeIconBrush, Asset.GetColor(), true);
+			FTableTreeNodePtr AssetGroupPtr = MakeShared<FAssetTreeNode>(Asset.GetNodeName(), AssetNode.GetAssetTableWeak(), AssetNode.GetRowIndex(), true);
 			AssetGroupPtr->SetExpansion(false);
-			ParentGroup.AddChildAndSetGroupPtr(AssetGroupPtr);
+			ParentGroup.AddChildAndSetParent(AssetGroupPtr);
 
 			// Add the asset node (self) under a "_self" group node.
 			static FName SelfGroupName(TEXT("_Self_")); // used _ prefix to sort before "Dependencies"
-			FTableTreeNodePtr SelfGroupPtr = MakeShared<FCustomTableTreeNode>(SelfGroupName, InParentTable, AssetNode.GetRowIndex(), AssetGroupNodeIconBrush, Asset.GetColor(), true);
+			FTableTreeNodePtr SelfGroupPtr = MakeShared<FAssetTreeNode>(SelfGroupName, AssetNode.GetAssetTableWeak(), AssetNode.GetRowIndex(), true);
 			SelfGroupPtr->SetExpansion(false);
-			SelfGroupPtr->AddChildAndSetGroupPtr(NodePtr);
-			AssetGroupPtr->AddChildAndSetGroupPtr(SelfGroupPtr);
+			SelfGroupPtr->AddChildAndSetParent(NodePtr);
+			AssetGroupPtr->AddChildAndSetParent(SelfGroupPtr);
 
 			// Create a group node for all dependent assets of the current asset.
 			// The actual nodes for the dependent assets are lazy created by this group node.
 			static FName DependenciesGroupName(TEXT("Dependencies"));
 			TSharedPtr<FAssetDependenciesGroupTreeNode> DependenciesGroupPtr = MakeShared<FAssetDependenciesGroupTreeNode>(DependenciesGroupName, AssetTable, AssetNode.GetRowIndex());
 			DependenciesGroupPtr->SetExpansion(false);
-			AssetGroupPtr->AddChildAndSetGroupPtr(DependenciesGroupPtr);
+			AssetGroupPtr->AddChildAndSetParent(DependenciesGroupPtr);
 		}
 		else
 		{
 			// No dependencies. Just add the asset node (self).
-			ParentGroup.AddChildAndSetGroupPtr(NodePtr);
+			ParentGroup.AddChildAndSetParent(NodePtr);
 		}
 	}
 }
diff --git a/Engine/Plugins/Editor/AssetManagerEditor/Source/AssetManagerEditor/Private/TreeView/AssetTable.cpp b/Engine/Plugins/Editor/AssetManagerEditor/Source/AssetManagerEditor/Private/TreeView/AssetTable.cpp
--- a/Engine/Plugins/Editor/AssetManagerEditor/Source/AssetManagerEditor/Private/TreeView/AssetTable.cpp
+++ b/Engine/Plugins/Editor/AssetManagerEditor/Source/AssetManagerEditor/Private/TreeView/AssetTable.cpp
@@ -831,7 +831,17 @@ void FAssetTable::AddDefaultColumns()
 		public:
 			virtual const TOptional<FTableCellValue> GetValue(const FTableColumn& Column, const FBaseTreeNode& Node) const override
 			{
-				if (Node.IsGroup())
+				if (Node.Is<FAssetTreeNode>() && Node.As<FAssetTreeNode>().IsValidAsset())
+				{
+					if (Node.Is<FAssetDependenciesGroupTreeNode>())
+					{
+						return TOptional<FTableCellValue>();
+					}
+					const FAssetTreeNode& TreeNode = Node.As<FAssetTreeNode>();
+					const FAssetTableRow& Asset = TreeNode.GetAssetChecked();
+					return FTableCellValue(static_cast<int64>(Asset.GetStagedCompressedSize()));
+				}
+				else if (Node.IsGroup())
 				{
 					const FTableTreeNode& NodePtr = static_cast<const FTableTreeNode&>(Node);
 					if (NodePtr.HasAggregatedValue(Column.GetId()))
@@ -839,12 +849,6 @@ void FAssetTable::AddDefaultColumns()
 						return NodePtr.GetAggregatedValue(Column.GetId());
 					}
 				}
-				else //if (Node->Is<FAssetTreeNode>())
-				{
-					const FAssetTreeNode& TreeNode = static_cast<const FAssetTreeNode&>(Node);
-					const FAssetTableRow& Asset = TreeNode.GetAssetChecked();
-					return FTableCellValue(static_cast<int64>(Asset.GetStagedCompressedSize()));
-				}
 
 				return TOptional<FTableCellValue>();
 			}
diff --git a/Engine/Plugins/Editor/AssetManagerEditor/Source/AssetManagerEditor/Private/TreeView/AssetTreeNode.cpp b/Engine/Plugins/Editor/AssetManagerEditor/Source/AssetManagerEditor/Private/TreeView/AssetTreeNode.cpp
--- a/Engine/Plugins/Editor/AssetManagerEditor/Source/AssetManagerEditor/Private/TreeView/AssetTreeNode.cpp
+++ b/Engine/Plugins/Editor/AssetManagerEditor/Source/AssetManagerEditor/Private/TreeView/AssetTreeNode.cpp
@@ -19,21 +19,11 @@ INSIGHTS_IMPLEMENT_RTTI(FAssetDependenciesGroupTreeNode)
 
 const FSlateBrush* FAssetTreeNode::GetIcon() const
 {
-#if 0 // debug / mock code
-	if (RowId.RowIndex % 7)
+	if (IsValidAsset())
 	{
-		return UE::Insights::FInsightsStyle::GetBrush("Icons.Package.TreeItem");
+		return FBaseTreeNode::GetDefaultIcon(false); // default icon for leaf nodes
 	}
-	if (RowId.RowIndex % 3)
-	{
-		return UE::Insights::FInsightsStyle::GetBrush("Icons.Asset.TreeItem");
-	}
-#endif
-
-	//const FAssetTableRow& AssetRpw = GetAssetChecked();
-	//TODO: switch (AssetRow->GetType()) --> icon
-
-	return FBaseTreeNode::GetDefaultIcon(false); // default icon for leaf nodes
+	return FBaseTreeNode::GetDefaultIcon(IsGroup()); // default icon
 }
 
 ////////////////////////////////////////////////////////////////////////////////////////////////////
@@ -47,6 +37,13 @@ FLinearColor FAssetTreeNode::GetColor() const
 // FAssetDependenciesGroupTreeNode
 ////////////////////////////////////////////////////////////////////////////////////////////////////
 
+const FSlateBrush* FAssetDependenciesGroupTreeNode::GetIcon() const
+{
+	return FBaseTreeNode::GetDefaultIcon(true); // default icon for group nodes
+}
+
+////////////////////////////////////////////////////////////////////////////////////////////////////
+
 FLinearColor FAssetDependenciesGroupTreeNode::GetColor() const
 {
 	return FLinearColor(0.75f, 0.5f, 1.0f, 1.0f);
@@ -78,19 +75,19 @@ bool FAssetDependenciesGroupTreeNode::OnLazyCreateChildren(TSharedPtr<UE::Insigh
 		return false;
 	}
 
-	TSharedPtr<FAssetTable> AssetTable = StaticCastSharedPtr<FAssetTable>(GetParentTable().Pin());
-	const FAssetTableRow& AssetRow = AssetTable->GetAssetChecked(GetRowIndex());
+	FAssetTable& AssetTable = GetAssetTableChecked();
+	const FAssetTableRow& AssetRow = AssetTable.GetAssetChecked(GetRowIndex());
 	const TArray<int32>& Dependencies = AssetRow.GetDependencies();
 	TArray<UE::Insights::FTableTreeNodePtr> AddedNodes;
 	for (int32 DepAssetIndex : Dependencies)
 	{
-		if (!ensure(AssetTable->IsValidRowIndex(DepAssetIndex)))
+		if (!ensure(AssetTable.IsValidRowIndex(DepAssetIndex)))
 		{
 			continue;
 		}
-		const FAssetTableRow& DepAssetRow = AssetTable->GetAssetChecked(DepAssetIndex);
+		const FAssetTableRow& DepAssetRow = AssetTable.GetAssetChecked(DepAssetIndex);
 		FName DepAssetNodeName(DepAssetRow.GetName());
-		FAssetTreeNodePtr DepAssetNodePtr = MakeShared<FAssetTreeNode>(DepAssetNodeName, AssetTable, DepAssetIndex);
+		FAssetTreeNodePtr DepAssetNodePtr = MakeShared<FAssetTreeNode>(DepAssetNodeName, GetAssetTableWeak(), DepAssetIndex);
 		AddedNodes.Add(DepAssetNodePtr);
 	}
 
diff --git a/Engine/Plugins/Editor/AssetManagerEditor/Source/AssetManagerEditor/Private/TreeView/AssetTreeNode.h b/Engine/Plugins/Editor/AssetManagerEditor/Source/AssetManagerEditor/Private/TreeView/AssetTreeNode.h
--- a/Engine/Plugins/Editor/AssetManagerEditor/Source/AssetManagerEditor/Private/TreeView/AssetTreeNode.h
+++ b/Engine/Plugins/Editor/AssetManagerEditor/Source/AssetManagerEditor/Private/TreeView/AssetTreeNode.h
@@ -14,37 +14,40 @@ class FAssetTreeNode : public UE::Insights::FTableTreeNode
 	INSIGHTS_DECLARE_RTTI(FAssetTreeNode, UE::Insights::FTableTreeNode)
 
 public:
-	/** Hidden constructor. */
-	//explicit FAssetTreeNode(const FName InName, TWeakPtr<UE::Insights::FTable> InParentTable, int32 InRowIndex) = delete;
-
-	/** Hidden constructor. */
-	//explicit FAssetTreeNode(const FName InGroupName, TWeakPtr<UE::Insights::FTable> InParentTable) = delete;
-
 	/** Initialization constructor for the asset node. */
 	explicit FAssetTreeNode(const FName InName, TWeakPtr<FAssetTable> InParentTable, int32 InRowIndex)
 		: FTableTreeNode(InName, InParentTable, InRowIndex)
-		, AssetTable(InParentTable.Pin().Get())
+		, AssetTablePtr(InParentTable.Pin().Get())
 	{
 	}
 
 	/** Initialization constructor for the group node. */
 	explicit FAssetTreeNode(const FName InGroupName, TWeakPtr<FAssetTable> InParentTable)
 		: FTableTreeNode(InGroupName, InParentTable)
-		, AssetTable(InParentTable.Pin().Get())
+		, AssetTablePtr(InParentTable.Pin().Get())
+	{
+	}
+
+	/** Initialization constructor for the asset and/or group node. */
+	explicit FAssetTreeNode(const FName InName, TWeakPtr<FAssetTable> InParentTable, int32 InRowIndex, bool bIsGroup)
+		: FTableTreeNode(InName, InParentTable, InRowIndex, bIsGroup)
+		, AssetTablePtr(InParentTable.Pin().Get())
 	{
 	}
 
 	virtual ~FAssetTreeNode() {}
 
-	bool IsValidAsset() const { return AssetTable && AssetTable->IsValidRowIndex(RowId.RowIndex); }
-	FAssetTable& GetAssetTableChecked() const { return *AssetTable; }
-	const FAssetTableRow& GetAssetChecked() const { return AssetTable->GetAssetChecked(RowId.RowIndex); }
+	TWeakPtr<FAssetTable> GetAssetTableWeak() const { return StaticCastWeakPtr<FAssetTable>(GetParentTable()); }
+
+	bool IsValidAsset() const { return AssetTablePtr && AssetTablePtr->IsValidRowIndex(RowId.RowIndex); }
+	FAssetTable& GetAssetTableChecked() const { return *AssetTablePtr; }
+	const FAssetTableRow& GetAssetChecked() const { return AssetTablePtr->GetAssetChecked(RowId.RowIndex); }
 
 	virtual const FSlateBrush* GetIcon() const override;
 	virtual FLinearColor GetColor() const override;
 
 private:
-	FAssetTable* AssetTable;
+	FAssetTable* AssetTablePtr;
 };
 
 /** Type definition for shared pointers to instances of FAssetTreeNode. */
@@ -61,25 +64,16 @@ typedef TWeakPtr<class FAssetTreeNode> FAssetTreeNodeWeak;
 
 ////////////////////////////////////////////////////////////////////////////////////////////////////
 
-class FAssetDependenciesGroupTreeNode : public UE::Insights::FTableTreeNode
+class FAssetDependenciesGroupTreeNode : public FAssetTreeNode
 {
-	INSIGHTS_DECLARE_RTTI(FAssetDependenciesGroupTreeNode, UE::Insights::FTableTreeNode)
+	INSIGHTS_DECLARE_RTTI(FAssetDependenciesGroupTreeNode, FAssetTreeNode)
 
 public:
-	/** Hidden constructor. */
-	//explicit FAssetDependenciesGroupTreeNode(const FName InName, TWeakPtr<UE::Insights::FTable> InParentTable, int32 InRowIndex) = delete;
-
-	/** Hidden constructor. */
-	//explicit FAssetDependenciesGroupTreeNode(const FName InGroupName, TWeakPtr<UE::Insights::FTable> InParentTable) = delete;
-
 	/** Initialization constructor for the group node. */
 	explicit FAssetDependenciesGroupTreeNode(const FName InGroupName, TWeakPtr<FAssetTable> InParentTable, int32 InParentRowIndex)
-		: FTableTreeNode(InGroupName, InParentTable, InParentRowIndex)
+		: FAssetTreeNode(InGroupName, InParentTable, InParentRowIndex, true)
 		, bAreChildrenCreated(false)
 	{
-		// This is actually a group node.
-		InitGroupData();
-
 		// Initially collapsed. Lazy create children when first expanded.
 		SetExpansion(false);
 	}
@@ -88,6 +82,7 @@ public:
 
 	virtual const FText GetExtraDisplayName() const;
 
+	virtual const FSlateBrush* GetIcon() const override;
 	virtual FLinearColor GetColor() const override;
 
 	virtual bool OnLazyCreateChildren(TSharedPtr<class UE::Insights::STableTreeView> InTableTreeView) override;
diff --git a/Engine/Plugins/Editor/AssetManagerEditor/Source/AssetManagerEditor/Private/TreeView/SAssetTableTreeView.cpp b/Engine/Plugins/Editor/AssetManagerEditor/Source/AssetManagerEditor/Private/TreeView/SAssetTableTreeView.cpp
--- a/Engine/Plugins/Editor/AssetManagerEditor/Source/AssetManagerEditor/Private/TreeView/SAssetTableTreeView.cpp
+++ b/Engine/Plugins/Editor/AssetManagerEditor/Source/AssetManagerEditor/Private/TreeView/SAssetTableTreeView.cpp
@@ -90,63 +90,61 @@ void SAssetTableTreeView::Tick(const FGeometry& AllottedGeometry, const double I
 
 void SAssetTableTreeView::RebuildTree(bool bResync)
 {
+	if (!bResync)
+	{
+		// There are no incremental updates.
+		return;
+	}
+
 	UE::Insights::FStopwatch Stopwatch;
 	Stopwatch.Start();
 
 	UE::Insights::FStopwatch SyncStopwatch;
 	SyncStopwatch.Start();
 
-	if (bResync)
-	{
-		TableTreeNodes.Empty();
-	}
+	CancelCurrentAsyncOp();
 
-	const int32 PreviousNodeCount = TableTreeNodes.Num();
+	const int32 PreviousNodeCount = TableRowNodes.Num();
+	TableRowNodes.Empty();
 
 	TSharedPtr<FAssetTable> AssetTable = GetAssetTable();
-	const TArray<FAssetTableRow>& Assets = AssetTable->GetAssets();
-	const int32 VisibleAssetCount = AssetTable->GetVisibleAssetCount();
+	const int32 VisibleAssetCount = AssetTable.IsValid() ? AssetTable->GetVisibleAssetCount() : 0;
 
-	if (VisibleAssetCount != PreviousNodeCount)
+	if (VisibleAssetCount > 0)
 	{
-		UE_LOG(LogInsights, Log, TEXT("[AssetTree] Creating nodes (%d nodes --> %d assets)..."), PreviousNodeCount, VisibleAssetCount);
+		UE_LOG(LogInsights, Log, TEXT("[AssetTree] Creating %d asset nodes (previously: %d nodes)..."), VisibleAssetCount, PreviousNodeCount);
 
-		if (VisibleAssetCount < PreviousNodeCount)
-		{
-			TableTreeNodes.Empty();
-		}
-		TableTreeNodes.Reserve(VisibleAssetCount);
+		TableRowNodes.Reserve(VisibleAssetCount);
 
-		//const FName BaseNodeName(TEXT("asset"));
-		for (int32 AssetIndex = TableTreeNodes.Num(); AssetIndex < VisibleAssetCount; ++AssetIndex)
+		ensure(TableRowNodes.Num() == 0);
+		for (int32 AssetIndex = 0; AssetIndex < VisibleAssetCount; ++AssetIndex)
 		{
 			const FAssetTableRow* Asset = AssetTable->GetAsset(AssetIndex);
 
-			//FName NodeName(BaseNodeName, AssetIndex + 1);
 			FName NodeName(Asset->GetName());
 			FAssetTreeNodePtr NodePtr = MakeShared<FAssetTreeNode>(NodeName, AssetTable, AssetIndex);
-			TableTreeNodes.Add(NodePtr);
+			TableRowNodes.Add(NodePtr);
 		}
-		ensure(TableTreeNodes.Num() == VisibleAssetCount);
+		ensure(TableRowNodes.Num() == VisibleAssetCount);
+	}
+	else
+	{
+		UE_LOG(LogInsights, Log, TEXT("[AssetTree] Resetting tree (previously: %d nodes)..."), PreviousNodeCount);
 	}
 
 	SyncStopwatch.Stop();
 
-	if (bResync || TableTreeNodes.Num() != PreviousNodeCount)
-	{
-		UE_LOG(LogInsights, Log, TEXT("[AssetTree] Update tree..."));
-		UpdateTree();
-		TreeView->RebuildList();
-		TreeView->ClearSelection();
-		SelectedIndices.Empty();
-		SelectedAssetNode.Reset();
-	}
+	UE_LOG(LogInsights, Log, TEXT("[AssetTree] Update tree..."));
+	UpdateTree();
+	TreeView->RebuildList();
+	TreeView->ClearSelection();
+	TreeView_OnSelectionChanged(nullptr, ESelectInfo::Type::Direct);
 
 	Stopwatch.Stop();
 	const double TotalTime = Stopwatch.GetAccumulatedTime();
 	const double SyncTime = SyncStopwatch.GetAccumulatedTime();
-	UE_LOG(LogInsights, Log, TEXT("[AssetTree] Tree view rebuilt in %.4fs (sync: %.4fs + update: %.4fs) --> %d nodes (%d added)"),
-		TotalTime, SyncTime, TotalTime - SyncTime, TableTreeNodes.Num(), TableTreeNodes.Num() - PreviousNodeCount);
+	UE_LOG(LogInsights, Log, TEXT("[AssetTree] Tree view rebuilt in %.4fs (sync: %.4fs + update: %.4fs) --> %d asset nodes"),
+		TotalTime, SyncTime, TotalTime - SyncTime, TableRowNodes.Num());
 }
 
 ////////////////////////////////////////////////////////////////////////////////////////////////////
@@ -558,19 +556,6 @@ void SAssetTableTreeView::InternalCreateGroupings()
 
 ////////////////////////////////////////////////////////////////////////////////////////////////////
 
-bool SAssetTableTreeView::ApplyCustomAdvancedFilters(const UE::Insights::FTableTreeNodePtr& NodePtr)
-{
-	return true;
-}
-
-////////////////////////////////////////////////////////////////////////////////////////////////////
-
-void SAssetTableTreeView::AddCustomAdvancedFilters()
-{
-}
-
-////////////////////////////////////////////////////////////////////////////////////////////////////
-
 void SAssetTableTreeView::ExportDependencyData() const
 {
 	IDesktopPlatform* DesktopPlatform = FDesktopPlatformModule::Get();
@@ -666,6 +651,8 @@ void SAssetTableTreeView::ExportDependencyData() const
 	}
 }
 
+////////////////////////////////////////////////////////////////////////////////////////////////////
+
 TArray<FAssetData> SAssetTableTreeView::GetAssetDataForSelection() const
 {
 	TArray<FAssetData> Assets;
@@ -690,6 +677,8 @@ TArray<FAssetData> SAssetTableTreeView::GetAssetDataForSelection() const
 	return Assets;
 }
 
+////////////////////////////////////////////////////////////////////////////////////////////////////
+
 void SAssetTableTreeView::ExtendMenu(FMenuBuilder& MenuBuilder)
 {
 	FCanExecuteAction HasSelectionAndCanExecute = FCanExecuteAction::CreateLambda([this]()
@@ -839,41 +828,46 @@ void SAssetTableTreeView::TreeView_OnSelectionChanged(UE::Insights::FTableTreeNo
 
 	for (const UE::Insights::FTableTreeNodePtr& Node : SelectedNodes)
 	{
-		if (Node->Is<FAssetTreeNode>() && !Node->IsGroup())
+		if (Node->Is<FAssetTreeNode>() && Node->As<FAssetTreeNode>().IsValidAsset())
 		{
 			NewSelectedAssetNode = StaticCastSharedPtr<FAssetTreeNode>(Node);
 			NewlySelectedAssetRowIndex = NewSelectedAssetNode->GetRowIndex();
 			SelectionSetIndices.Add(NewlySelectedAssetRowIndex);
 			++NumSelectedAssets;
 		}
-		else if (Node->Is<UE::Insights::FCustomTableTreeNode>() && Node->IsGroup())
-		{
-			int32 RowIndex = Node->GetRowIndex();
-			if (RowIndex != UE::Insights::FTableRowId::InvalidRowIndex)
-			{
-				SelectionSetIndices.Add(RowIndex);
-				NewlySelectedAssetRowIndex = RowIndex;
-				++NumSelectedAssets;
-			}
-		}
 	}
 
+	const int32 FilteredAssetCount = FilteredNodesPtr->Num();
 	const int32 VisibleAssetCount = GetAssetTable()->GetVisibleAssetCount();
+
 	if (NumSelectedAssets == 0)
 	{
-		FooterLeftText = FText::Format(LOCTEXT("FooterLeftTextFmt0", "{0} assets"), FText::AsNumber(VisibleAssetCount));
+		if (FilteredAssetCount != VisibleAssetCount)
+		{
+			FooterLeftText = FText::Format(LOCTEXT("FooterLeftTextFmt_NoSelected_Filtered", "{0} / {1} assets"), FText::AsNumber(FilteredAssetCount), FText::AsNumber(VisibleAssetCount));
+		}
+		else
+		{
+			FooterLeftText = FText::Format(LOCTEXT("FooterLeftTextFmt_NoSelected_NoFiltered", "{0} assets"), FText::AsNumber(VisibleAssetCount));
+		}
 		FooterCenterText1 = FText();
 		FooterCenterText2 = FText();
 		FooterRightText1 = FText();
 	}
 	else if (NumSelectedAssets == 1)
 	{
-		FooterLeftText = FText::Format(LOCTEXT("FooterLeftTextFmt1", "{0} assets (1 selected)"), FText::AsNumber(VisibleAssetCount));
-
+		if (FilteredAssetCount != VisibleAssetCount)
+		{
+			FooterLeftText = FText::Format(LOCTEXT("FooterLeftTextFmt_1Selected_Filtered", "{0}/{1} assets (1 selected)"), FText::AsNumber(FilteredAssetCount), FText::AsNumber(VisibleAssetCount));
+		}
+		else
+		{
+			FooterLeftText = FText::Format(LOCTEXT("FooterLeftTextFmt_1Selected_NoFiltered", "{0} assets (1 selected)"), FText::AsNumber(VisibleAssetCount));
+		}
 		const FAssetTableRow& AssetTableRow = GetAssetTable()->GetAssetChecked(NewlySelectedAssetRowIndex);
 		FooterCenterText1 = FText::FromString(AssetTableRow.GetPath());
 		FooterCenterText2 = FText::FromString(AssetTableRow.GetName());
-		FooterRightText1 = FText::Format(LOCTEXT("FooterRightFmt3", "Self: {0} Unique: {1} Shared: {2} External: {3}"),
+		FooterRightText1 = FText::Format(LOCTEXT("FooterRightFmt", "Self: {0}    Unique: {1}    Shared: {2}    External: {3}"),
 			FText::AsMemory(AssetTableRow.GetStagedCompressedSize()),
 			FText::AsMemory(AssetTableRow.GetOrComputeTotalSizeUniqueDependencies(*GetAssetTable(), NewlySelectedAssetRowIndex)),
 			FText::AsMemory(AssetTableRow.GetOrComputeTotalSizeSharedDependencies(*GetAssetTable(), NewlySelectedAssetRowIndex)),
@@ -881,7 +875,14 @@ void SAssetTableTreeView::TreeView_OnSelectionChanged(UE::Insights::FTableTreeNo
 	}
 	else
 	{
-		FooterLeftText = FText::Format(LOCTEXT("FooterLeftTextFmt2", "{0} assets ({1} selected)"), FText::AsNumber(VisibleAssetCount), FText::AsNumber(NumSelectedAssets));
+		if (FilteredAssetCount != VisibleAssetCount)
+		{
+			FooterLeftText = FText::Format(LOCTEXT("FooterLeftTextFmt_ManySelected_Filtered", "{0} / {1} assets ({2} selected)"), FText::AsNumber(FilteredAssetCount), FText::AsNumber(VisibleAssetCount), FText::AsNumber(NumSelectedAssets));
+		}
+		else
+		{
+			FooterLeftText = FText::Format(LOCTEXT("FooterLeftTextFmt_ManySelected_NoFiltered", "{0} assets ({1} selected)"), FText::AsNumber(VisibleAssetCount), FText::AsNumber(NumSelectedAssets));
+		}
 		FooterCenterText1 = FText();
 		FooterCenterText2 = FText();
 
@@ -893,7 +894,7 @@ void SAssetTableTreeView::TreeView_OnSelectionChanged(UE::Insights::FTableTreeNo
 
 		int64 TotalExternalDependencySize = FAssetTableRow::ComputeTotalSizeExternalDependencies(*GetAssetTable(), SelectionSetIndices);
 		FAssetTableDependencySizes Sizes = FAssetTableRow::ComputeDependencySizes(*GetAssetTable(), SelectionSetIndices, nullptr, nullptr);
-		FooterRightText1 = FText::Format(LOCTEXT("FooterRightFmt3", "Self: {0} Unique: {1} Shared: {2} External: {3}"),
+		FooterRightText1 = FText::Format(LOCTEXT("FooterRightFmt", "Self: {0}    Unique: {1}    Shared: {2}    External: {3}"),
 			FText::AsMemory(TotalSelfSize),
 			FText::AsMemory(Sizes.UniqueDependenciesSize),
 			FText::AsMemory(Sizes.SharedDependenciesSize),
diff --git a/Engine/Plugins/Editor/AssetManagerEditor/Source/AssetManagerEditor/Private/TreeView/SAssetTableTreeView.h b/Engine/Plugins/Editor/AssetManagerEditor/Source/AssetManagerEditor/Private/TreeView/SAssetTableTreeView.h
--- a/Engine/Plugins/Editor/AssetManagerEditor/Source/AssetManagerEditor/Private/TreeView/SAssetTableTreeView.h
+++ b/Engine/Plugins/Editor/AssetManagerEditor/Source/AssetManagerEditor/Private/TreeView/SAssetTableTreeView.h
@@ -84,9 +84,6 @@ protected:
 	virtual void TreeView_OnSelectionChanged(UE::Insights::FTableTreeNodePtr SelectedItem, ESelectInfo::Type SelectInfo) override;
 
 private:
-	bool virtual ApplyCustomAdvancedFilters(const UE::Insights::FTableTreeNodePtr& NodePtr) override;
-	virtual void AddCustomAdvancedFilters() override;
-
 	void InitAvailableViewPresets();
 
 	void ExportDependencyData() const;
