Asset Audit Tool: Changing GameFeaturePlugin references to Plugin references for generality.

diff --git a/Engine/Plugins/Editor/AssetManagerEditor/Source/AssetManagerEditor/Private/AssetManagerEditorModule.cpp b/Engine/Plugins/Editor/AssetManagerEditor/Source/AssetManagerEditor/Private/AssetManagerEditorModule.cpp
--- a/Engine/Plugins/Editor/AssetManagerEditor/Source/AssetManagerEditor/Private/AssetManagerEditorModule.cpp
+++ b/Engine/Plugins/Editor/AssetManagerEditor/Source/AssetManagerEditor/Private/AssetManagerEditorModule.cpp
@@ -109,7 +109,7 @@ const FName IAssetManagerEditorModule::CookRuleName = FName("CookRule");
 const FName IAssetManagerEditorModule::ChunksName = FName("Chunks");
 const FName IAssetManagerEditorModule::StageChunkSizeName = FName("Stage_ChunkSize");
 const FName IAssetManagerEditorModule::StageChunkCompressedSizeName = FName("Stage_ChunkCompressedSize");
-const FName IAssetManagerEditorModule::GameFeaturePluginsName = FName("GameFeaturePlugins");
+const FName IAssetManagerEditorModule::PluginName = FName("GameFeaturePlugins");
 
 const FString FAssetManagerEditorRegistrySource::EditorSourceName = TEXT("Editor");
 const FString FAssetManagerEditorRegistrySource::CustomSourceName = TEXT("Custom");
@@ -1310,7 +1310,7 @@ bool FAssetManagerEditorModule::GetStringValueForCustomColumn(const FAssetData&
 		}
 		return true;
 	}
-	else if (ColumnName == GameFeaturePluginsName)
+	else if (ColumnName == PluginName)
 	{
 		TArray<FString> ParsedPath;
 		const FString& PackageNameString = AssetData.PackageName.ToString();
@@ -1386,7 +1386,7 @@ bool FAssetManagerEditorModule::GetDisplayTextForCustomColumn(const FAssetData&
 			return true;
 		}
 	}
-	else if (ColumnName == ChunksName || ColumnName == GameFeaturePluginsName)
+	else if (ColumnName == ChunksName || ColumnName == PluginName)
 	{
 		FString OutString;
 
diff --git a/Engine/Plugins/Editor/AssetManagerEditor/Source/AssetManagerEditor/Private/SAssetAuditBrowser.cpp b/Engine/Plugins/Editor/AssetManagerEditor/Source/AssetManagerEditor/Private/SAssetAuditBrowser.cpp
--- a/Engine/Plugins/Editor/AssetManagerEditor/Source/AssetManagerEditor/Private/SAssetAuditBrowser.cpp
+++ b/Engine/Plugins/Editor/AssetManagerEditor/Source/AssetManagerEditor/Private/SAssetAuditBrowser.cpp
@@ -797,9 +797,9 @@ void SAssetAuditBrowser::PopulateAssetTableRow(FAssetTableRow& OutRow, const FAs
 		OutRow.Chunks = AssetTable.StoreStr(Str);
 	}
 
-	if (EditorModule->GetStringValueForCustomColumn(AssetData, IAssetManagerEditorModule::GameFeaturePluginsName, Str))
+	if (EditorModule->GetStringValueForCustomColumn(AssetData, IAssetManagerEditorModule::PluginName, Str))
 	{
-		OutRow.GameFeaturePlugin = AssetTable.StoreStr(Str);
+		OutRow.PluginName = AssetTable.StoreStr(Str);
 	}
 
 	OutRow.NativeClass = AssetTable.StoreStr(AssetData.AssetClassPath.GetAssetName().ToString());
diff --git a/Engine/Plugins/Editor/AssetManagerEditor/Source/AssetManagerEditor/Private/TreeView/AssetTable.cpp b/Engine/Plugins/Editor/AssetManagerEditor/Source/AssetManagerEditor/Private/TreeView/AssetTable.cpp
--- a/Engine/Plugins/Editor/AssetManagerEditor/Source/AssetManagerEditor/Private/TreeView/AssetTable.cpp
+++ b/Engine/Plugins/Editor/AssetManagerEditor/Source/AssetManagerEditor/Private/TreeView/AssetTable.cpp
@@ -34,7 +34,7 @@ const FName FAssetTableColumns::TotalSizeExternalDependenciesColumnId(TEXT("Tota
 const FName FAssetTableColumns::TotalUsageCountColumnId(TEXT("TotalUsageCount"));
 const FName FAssetTableColumns::ChunksColumnId(TEXT("Chunks"));
 const FName FAssetTableColumns::NativeClassColumnId(TEXT("NativeClass"));
-const FName FAssetTableColumns::GameFeaturePluginColumnId(TEXT("GameFeaturePlugin"));
+const FName FAssetTableColumns::PluginNameColumnId(TEXT("PluginName"));
 
 ////////////////////////////////////////////////////////////////////////////////////////////////////
 // FAssetTableStringStore
@@ -170,7 +170,7 @@ FAssetTable::FAssetTable()
 		Asset.TotalUsageCount = 10;
 		Asset.StagedCompressedSize = 1;
 		Asset.NativeClass = StoreStr(FString::Printf(TEXT("NativeClass%02d"), (Id * Id * Id) % 8));
-		Asset.GameFeaturePlugin = StoreStr(TEXT("MockGFP"));
+		Asset.PluginName = StoreStr(TEXT("MockGFP"));
 	}
 
 	for (int32 TestIndex = 0; TestIndex < 5; TestIndex++)
@@ -1235,16 +1235,16 @@ void FAssetTable::AddDefaultColumns()
 		AddColumn(ColumnRef);
 	}
 	//////////////////////////////////////////////////
-	// GameFeaturePlugin Column
+	// PluginName Column
 	{
-		TSharedRef<FTableColumn> ColumnRef = MakeShared<FTableColumn>(FAssetTableColumns::GameFeaturePluginColumnId);
+		TSharedRef<FTableColumn> ColumnRef = MakeShared<FTableColumn>(FAssetTableColumns::PluginNameColumnId);
 		FTableColumn& Column = *ColumnRef;
 
 		Column.SetIndex(ColumnIndex++);
 
-		Column.SetShortName(LOCTEXT("GameFeaturePluginColumnName", "GameFeaturePlugin"));
-		Column.SetTitleName(LOCTEXT("GameFeaturePluginColumnTitle", "GameFeaturePlugin"));
-		Column.SetDescription(LOCTEXT("GameFeaturePluginColumnDesc", "GameFeaturePlugin of the asset"));
+		Column.SetShortName(LOCTEXT("PluginColumnName", "Plugin"));
+		Column.SetTitleName(LOCTEXT("PluginColumnTitle", "Plugin"));
+		Column.SetDescription(LOCTEXT("PluginColumnDesc", "Plugin (Game Feature or Engine) of the asset"));
 
 		Column.SetFlags(ETableColumnFlags::ShouldBeVisible | ETableColumnFlags::CanBeHidden | ETableColumnFlags::CanBeFiltered);
 
@@ -1270,7 +1270,7 @@ void FAssetTable::AddDefaultColumns()
 				{
 					const FAssetTreeNode& TreeNode = static_cast<const FAssetTreeNode&>(Node);
 					const FAssetTableRow& Asset = TreeNode.GetAssetChecked();
-					return FTableCellValue(Asset.GetGameFeaturePlugin());
+					return FTableCellValue(Asset.GetPluginName());
 				}
 
 				return TOptional<FTableCellValue>();
@@ -1292,7 +1292,7 @@ void FAssetTable::AddDefaultColumns()
 	//////////////////////////////////////////////////
 }
 
-/*static*/TSet<int32> FAssetTableRow::GatherAllReachableNodes(const TArray<int32>& StartingNodes, const FAssetTable& OwningTable, const TSet<int32>& AdditionalNodesToStopAt, const TSet<const TCHAR*, TStringPointerSetKeyFuncs_DEPRECATED<const TCHAR*>>& RestrictToGFPs)
+/*static*/TSet<int32> FAssetTableRow::GatherAllReachableNodes(const TArray<int32>& StartingNodes, const FAssetTable& OwningTable, const TSet<int32>& AdditionalNodesToStopAt, const TSet<const TCHAR*, TStringPointerSetKeyFuncs_DEPRECATED<const TCHAR*>>& RestrictToPlugins)
 {
 	// "visit" ThisIndex to seed the exploration
 	TSet<int32> VisitedIndices;
@@ -1313,7 +1313,7 @@ void FAssetTable::AddDefaultColumns()
 		Iterator.RemoveCurrent();
 
 		const FAssetTableRow& Row = OwningTable.GetAssetChecked(CurrentIndex);
-		if (ShouldSkipDueToGFP(RestrictToGFPs, Row.GetGameFeaturePlugin()) || AdditionalNodesToStopAt.Contains(CurrentIndex))
+		if (ShouldSkipDueToPlugin(RestrictToPlugins, Row.GetPluginName()) || AdditionalNodesToStopAt.Contains(CurrentIndex))
 		{
 			// Don't traverse outside this plugin
 			continue;
@@ -1352,7 +1352,7 @@ void FAssetTable::AddDefaultColumns()
 		Iterator.RemoveCurrent();
 
 		const FAssetTableRow& Row = OwningTable.GetAssetChecked(CurrentIndex);
-		if (ShouldSkipDueToGFP(RestrictToGFPs, Row.GetGameFeaturePlugin()))
+		if (ShouldSkipDueToPlugin(RestrictToGFPs, Row.GetPluginName()))
 		{
 			// Don't traverse outside this plugin
 			continue;
@@ -1405,21 +1405,21 @@ void FAssetTable::AddDefaultColumns()
 
 	FAssetTableDependencySizes Result;
 
-	TSet<const TCHAR*, TStringPointerSetKeyFuncs_DEPRECATED<const TCHAR*>> RestrictToGFPs;
+	TSet<const TCHAR*, TStringPointerSetKeyFuncs_DEPRECATED<const TCHAR*>> RestrictToPlugins;
 	for (int32 RootIndex : RootIndices)
 	{
-		RestrictToGFPs.Add(OwningTable.GetAssetChecked(RootIndex).GetGameFeaturePlugin());
+		RestrictToPlugins.Add(OwningTable.GetAssetChecked(RootIndex).GetPluginName());
 	}
 
 
-	TSet<int32> VisitedIndices = GatherAllReachableNodes(RootIndices.Array(), OwningTable, TSet<int32>(), RestrictToGFPs);
+	TSet<int32> VisitedIndices = GatherAllReachableNodes(RootIndices.Array(), OwningTable, TSet<int32>(), RestrictToPlugins);
 
 	// Iteratively separate the graph of "all things referenced by ThisIndex, directly or indirectly"
 	// into "UniqueDependencies -- things referenced ONLY by ThisIndex and by other things themselves referenced ONLY by ThisIndex" limited to the GFPs of the RootIndices
 	// and "SharedDependencies" -- things removed from the list of "all things referenced by ThisIndex" in order to identify UniqueDependencies, limited to the GFPs of the RootIndices
 	TSet<int32> UniqueDependencies;
 	TSet<int32> SharedDependencies;
-	RefineDependencies(VisitedIndices, OwningTable, RootIndices, RestrictToGFPs, UniqueDependencies, SharedDependencies);
+	RefineDependencies(VisitedIndices, OwningTable, RootIndices, RestrictToPlugins, UniqueDependencies, SharedDependencies);
 
 	// If there's only one root provided, the dependencies shouldn't include it. If more than one root is provided,
 	// some roots might be included as dependencies of other roots
@@ -1453,7 +1453,7 @@ void FAssetTable::AddDefaultColumns()
 
 		const FAssetTableRow& Row = OwningTable.GetAssetChecked(CurrentIndex);
 		VisitedIndices.Add(CurrentIndex);
-		if (ShouldSkipDueToGFP(RestrictToGFPs, Row.GetGameFeaturePlugin()))
+		if (ShouldSkipDueToPlugin(RestrictToPlugins, Row.GetPluginName()))
 		{
 			// Don't traverse outside this plugin
 			continue;
@@ -1486,10 +1486,10 @@ void FAssetTable::AddDefaultColumns()
 	Stopwatch.Start();
 
 	TSet<int32> AllReachableNodes = GatherAllReachableNodes(StartingNodes.Array(), OwningTable, TSet<int32>{}, TSet<const TCHAR*, TStringPointerSetKeyFuncs_DEPRECATED<const TCHAR*>>{});
-	TSet<const TCHAR*, TStringPointerSetKeyFuncs_DEPRECATED<const TCHAR*>> SourceGFPs;
+	TSet<const TCHAR*, TStringPointerSetKeyFuncs_DEPRECATED<const TCHAR*>> SourcePlugins;
 	for (int32 RootIndex : StartingNodes)
 	{
-		SourceGFPs.Add(OwningTable.GetAsset(RootIndex)->GetGameFeaturePlugin());
+		SourcePlugins.Add(OwningTable.GetAsset(RootIndex)->GetPluginName());
 	}
 
 	int64 TotalSizeExternalDependencies = 0;
@@ -1498,7 +1498,7 @@ void FAssetTable::AddDefaultColumns()
 	{
 		const FAssetTableRow& Row = *OwningTable.GetAsset(Index);
 		// Only include EXTERNAL dependencies
-		if (!SourceGFPs.Contains(Row.GameFeaturePlugin))
+		if (!SourcePlugins.Contains(Row.PluginName))
 		{
 			if (OutExternalDependencies != nullptr)
 			{
diff --git a/Engine/Plugins/Editor/AssetManagerEditor/Source/AssetManagerEditor/Private/TreeView/AssetTable.h b/Engine/Plugins/Editor/AssetManagerEditor/Source/AssetManagerEditor/Private/TreeView/AssetTable.h
--- a/Engine/Plugins/Editor/AssetManagerEditor/Source/AssetManagerEditor/Private/TreeView/AssetTable.h
+++ b/Engine/Plugins/Editor/AssetManagerEditor/Source/AssetManagerEditor/Private/TreeView/AssetTable.h
@@ -123,7 +123,7 @@ struct FAssetTableColumns
 	static const FName TotalUsageCountColumnId;
 	static const FName ChunksColumnId;
 	static const FName NativeClassColumnId;
-	static const FName GameFeaturePluginColumnId;
+	static const FName PluginNameColumnId;
 };
 
 struct FAssetTableDependencySizes
@@ -159,7 +159,7 @@ public:
 	int64 GetTotalUsageCount() const { return TotalUsageCount; }
 	const TCHAR* GetChunks() const { return Chunks; }
 	const TCHAR* GetNativeClass() const { return NativeClass; }
-	const TCHAR* GetGameFeaturePlugin() const { return GameFeaturePlugin; }
+	const TCHAR* GetPluginName() const { return PluginName; }
 	int32 GetNumDependencies() const { return Dependencies.Num(); }
 	const TArray<int32>& GetDependencies() const { return Dependencies; }
 	int32 GetNumReferencers() const { return Referencers.Num(); }
@@ -173,10 +173,10 @@ public:
 
 private:
 	static void RefineDependencies(const TSet<int32>& PreviouslyVisitedIndices, const FAssetTable& OwningTable, const TSet<int32> RootIndices, const TSet<const TCHAR*, TStringPointerSetKeyFuncs_DEPRECATED<const TCHAR*>>& RestrictToGFPs, TSet<int32>& OutUniqueIndices, TSet<int32>& OutSharedIndices);
-	static inline bool ShouldSkipDueToGFP(const TSet<const TCHAR*, TStringPointerSetKeyFuncs_DEPRECATED<const TCHAR*>>& RestrictToGFPs, const TCHAR* GFP) { return (RestrictToGFPs.Num() > 0) && !RestrictToGFPs.Contains(GFP); }
+	static inline bool ShouldSkipDueToPlugin(const TSet<const TCHAR*, TStringPointerSetKeyFuncs_DEPRECATED<const TCHAR*>>& RestrictToPlugins, const TCHAR* PluginName) { return (RestrictToPlugins.Num() > 0) && !RestrictToPlugins.Contains(PluginName); }
 
 	// Finds all nodes reachable from StartingNodes INCLUDING those in StartingNodes
-	static TSet<int32> GatherAllReachableNodes(const TArray<int32>& StartingNodes, const FAssetTable& OwningTable, const TSet<int32>& AdditionalNodesToStopAt, const TSet<const TCHAR*, TStringPointerSetKeyFuncs_DEPRECATED<const TCHAR*>>& RestrictToGFPs);
+	static TSet<int32> GatherAllReachableNodes(const TArray<int32>& StartingNodes, const FAssetTable& OwningTable, const TSet<int32>& AdditionalNodesToStopAt, const TSet<const TCHAR*, TStringPointerSetKeyFuncs_DEPRECATED<const TCHAR*>>& RestrictToPlugins);
 
 private:
 	FSoftObjectPath SoftObjectPath; // Used to allow you to go to that asset in the content browser or to open its editor after lookingup the FAssetData in the original registry state
@@ -192,7 +192,7 @@ private:
 	int64 TotalUsageCount = 0;
 	const TCHAR* Chunks = nullptr;
 	const TCHAR* NativeClass = nullptr;
-	const TCHAR* GameFeaturePlugin = nullptr;
+	const TCHAR* PluginName = nullptr;
 	TArray<int32> Dependencies;
 	TArray<int32> Referencers;
 	FLinearColor Color;
diff --git a/Engine/Plugins/Editor/AssetManagerEditor/Source/AssetManagerEditor/Private/TreeView/SAssetTableTreeView.cpp b/Engine/Plugins/Editor/AssetManagerEditor/Source/AssetManagerEditor/Private/TreeView/SAssetTableTreeView.cpp
--- a/Engine/Plugins/Editor/AssetManagerEditor/Source/AssetManagerEditor/Private/TreeView/SAssetTableTreeView.cpp
+++ b/Engine/Plugins/Editor/AssetManagerEditor/Source/AssetManagerEditor/Private/TreeView/SAssetTableTreeView.cpp
@@ -256,7 +256,7 @@ void SAssetTableTreeView::InitAvailableViewPresets()
 			InOutConfigSet.Add({ FAssetTableColumns::TotalUsageCountColumnId,               true, 100.0f });
 			InOutConfigSet.Add({ FAssetTableColumns::ChunksColumnId,                       !true, 200.0f });
 			InOutConfigSet.Add({ FAssetTableColumns::NativeClassColumnId,                   true, 200.0f });
-			InOutConfigSet.Add({ FAssetTableColumns::GameFeaturePluginColumnId,             true, 200.0f });
+			InOutConfigSet.Add({ FAssetTableColumns::PluginNameColumnId,				    true, 200.0f });
 		}
 	};
 	AvailableViewPresets.Add(MakeShared<FDefaultViewPreset>());
@@ -295,7 +295,7 @@ void SAssetTableTreeView::InitAvailableViewPresets()
 				[](TSharedPtr<UE::Insights::FTreeNodeGrouping>& Grouping)
 				{
 					return Grouping->Is<UE::Insights::FTreeNodeGroupingByUniqueValueCString>() &&
-						Grouping->As<UE::Insights::FTreeNodeGroupingByUniqueValueCString>().GetColumnId() == FAssetTableColumns::GameFeaturePluginColumnId;
+						Grouping->As<UE::Insights::FTreeNodeGroupingByUniqueValueCString>().GetColumnId() == FAssetTableColumns::PluginNameColumnId;
 				});
 			if (GameFeaturePluginGrouping)
 			{
@@ -339,7 +339,7 @@ void SAssetTableTreeView::InitAvailableViewPresets()
 			InOutConfigSet.Add({ FAssetTableColumns::TotalUsageCountColumnId,              !true, 100.0f });
 			InOutConfigSet.Add({ FAssetTableColumns::ChunksColumnId,                       !true, 200.0f });
 			InOutConfigSet.Add({ FAssetTableColumns::NativeClassColumnId,                   true, 200.0f });
-			InOutConfigSet.Add({ FAssetTableColumns::GameFeaturePluginColumnId,             true, 200.0f });
+			InOutConfigSet.Add({ FAssetTableColumns::PluginNameColumnId,             true, 200.0f });
 		}
 	};
 	AvailableViewPresets.Add(MakeShared<FGameFeaturePluginTypeDependencyView>());
@@ -400,7 +400,7 @@ void SAssetTableTreeView::InitAvailableViewPresets()
 			InOutConfigSet.Add({ FAssetTableColumns::TotalUsageCountColumnId,               true, 100.0f });
 			InOutConfigSet.Add({ FAssetTableColumns::ChunksColumnId,                       !true, 200.0f });
 			InOutConfigSet.Add({ FAssetTableColumns::NativeClassColumnId,                   true, 200.0f });
-			InOutConfigSet.Add({ FAssetTableColumns::GameFeaturePluginColumnId,             true, 200.0f });
+			InOutConfigSet.Add({ FAssetTableColumns::PluginNameColumnId,             true, 200.0f });
 		}
 	};
 	AvailableViewPresets.Add(MakeShared<FAssetPathViewPreset>());
@@ -461,7 +461,7 @@ void SAssetTableTreeView::InitAvailableViewPresets()
 			InOutConfigSet.Add({ FAssetTableColumns::TotalUsageCountColumnId,               true, 100.0f });
 			InOutConfigSet.Add({ FAssetTableColumns::ChunksColumnId,                       !true, 200.0f });
 			InOutConfigSet.Add({ FAssetTableColumns::NativeClassColumnId,                   true, 200.0f });
-			InOutConfigSet.Add({ FAssetTableColumns::GameFeaturePluginColumnId,             true, 200.0f });
+			InOutConfigSet.Add({ FAssetTableColumns::PluginNameColumnId,             true, 200.0f });
 		}
 	};
 	AvailableViewPresets.Add(MakeShared<FPrimaryAssetViewPreset>());
diff --git a/Engine/Plugins/Editor/AssetManagerEditor/Source/AssetManagerEditor/Public/AssetManagerEditorModule.h b/Engine/Plugins/Editor/AssetManagerEditor/Source/AssetManagerEditor/Public/AssetManagerEditorModule.h
--- a/Engine/Plugins/Editor/AssetManagerEditor/Source/AssetManagerEditor/Public/AssetManagerEditorModule.h
+++ b/Engine/Plugins/Editor/AssetManagerEditor/Source/AssetManagerEditor/Public/AssetManagerEditorModule.h
@@ -383,7 +383,7 @@ public:
 	static const FName ChunksName;
 	static const FName StageChunkSizeName;
 	static const FName StageChunkCompressedSizeName;
-	static const FName GameFeaturePluginsName;
+	static const FName PluginName;
 
 	/** Gets the value of a "virtual" column for an asset data, this will query the AssetManager for you and takes current platform into account. Returns true and sets out parameter if found */
 	virtual bool GetStringValueForCustomColumn(const FAssetData& AssetData, FName ColumnName, FString& OutValue) = 0;
