AssetAuditTreeView now provides base and marginal cost calculations when selecting a group of assets of the same type. Exporting dependencies now includes a column containing a dependency chain from a root asset to the relevant row.

diff --git a/Engine/Plugins/Editor/AssetManagerEditor/Source/AssetManagerEditor/Private/TreeView/AssetTable.cpp b/Engine/Plugins/Editor/AssetManagerEditor/Source/AssetManagerEditor/Private/TreeView/AssetTable.cpp
--- a/Engine/Plugins/Editor/AssetManagerEditor/Source/AssetManagerEditor/Private/TreeView/AssetTable.cpp
+++ b/Engine/Plugins/Editor/AssetManagerEditor/Source/AssetManagerEditor/Private/TreeView/AssetTable.cpp
@@ -1292,7 +1292,7 @@ void FAssetTable::AddDefaultColumns()
 	//////////////////////////////////////////////////
 }
 
-/*static*/TSet<int32> FAssetTableRow::GatherAllReachableNodes(const TArray<int32>& StartingNodes, const FAssetTable& OwningTable, const TSet<int32>& AdditionalNodesToStopAt, const TSet<const TCHAR*, TStringPointerSetKeyFuncs_DEPRECATED<const TCHAR*>>& RestrictToPlugins)
+/*static*/TSet<int32> FAssetTableRow::GatherAllReachableNodes(const TArray<int32>& StartingNodes, const FAssetTable& OwningTable, const TSet<int32>& AdditionalNodesToStopAt, const TSet<const TCHAR*, TStringPointerSetKeyFuncs_DEPRECATED<const TCHAR*>>& RestrictToPlugins, TMap<int32, TArray<int32>>* OutRouteMap /*= nullptr*/)
 {
 	// "visit" ThisIndex to seed the exploration
 	TSet<int32> VisitedIndices;
@@ -1300,9 +1300,18 @@ void FAssetTable::AddDefaultColumns()
 	for (int32 StartingNodeIndex : StartingNodes)
 	{
 		VisitedIndices.Add(StartingNodeIndex);
+		if (OutRouteMap != nullptr)
+		{
+			OutRouteMap->Add(StartingNodeIndex, TArray<int32>{StartingNodeIndex});
+		}
 		for (int32 ChildIndex : OwningTable.GetAssetChecked(StartingNodeIndex).GetDependencies())
 		{
 			IndicesToVisit.Add(ChildIndex);
+			if (OutRouteMap != nullptr)
+			{
+				TArray<int32> Route{StartingNodeIndex, ChildIndex};
+				OutRouteMap->Add(ChildIndex, Route);
+			}
 		}
 	}
 
@@ -1315,16 +1324,33 @@ void FAssetTable::AddDefaultColumns()
 		const FAssetTableRow& Row = OwningTable.GetAssetChecked(CurrentIndex);
 		if (ShouldSkipDueToPlugin(RestrictToPlugins, Row.GetPluginName()) || AdditionalNodesToStopAt.Contains(CurrentIndex))
 		{
+			if (OutRouteMap != nullptr)
+			{
+				OutRouteMap->Remove(CurrentIndex);
+			}
 			// Don't traverse outside this plugin
 			continue;
 		}
 
 		VisitedIndices.Add(CurrentIndex);
 
+		TArray<int32>* ParentRoute = nullptr;
+		if (OutRouteMap != nullptr)
+		{
+			ParentRoute = OutRouteMap->Find(CurrentIndex);
+		}
+
 		for (int32 ChildIndex : OwningTable.GetAssetChecked(CurrentIndex).GetDependencies())
 		{
 			if (!VisitedIndices.Contains(ChildIndex))
 			{
+				if (OutRouteMap != nullptr && ensure(ParentRoute != nullptr))
+				{
+					TArray<int32> ChildRoute(*ParentRoute, 1);
+					ChildRoute.Add(ChildIndex);
+
+					OutRouteMap->Add(ChildIndex, ChildRoute);
+				}
 				IndicesToVisit.Add(ChildIndex);
 			}
 		}
@@ -1341,7 +1367,11 @@ void FAssetTable::AddDefaultColumns()
 	TSet<int32> IndicesToVisit;
 	for (int32 RootIndex : RootIndices)
 	{
-		IndicesToVisit = IndicesToVisit.Union(TSet<int32>{OwningTable.GetAssetChecked(RootIndex).GetDependencies()});
+		const TArray<int32>& Dependencies = OwningTable.GetAssetChecked(RootIndex).GetDependencies();
+		for (int32 DependencyIndex : Dependencies)
+		{
+			IndicesToVisit.Add(DependencyIndex);
+		}
 	}
 	TSet<int32> NewlyExcludedDependencies;
 
@@ -1480,12 +1510,12 @@ void FAssetTable::AddDefaultColumns()
 
 ////////////////////////////////////////////////////////////////////////////////////////////////////
 
-/*static*/ int64 FAssetTableRow::ComputeTotalSizeExternalDependencies(const FAssetTable& OwningTable, const TSet<int32>& StartingNodes, TSet<int32>* OutExternalDependencies/* = nullptr*/)
+/*static*/ int64 FAssetTableRow::ComputeTotalSizeExternalDependencies(const FAssetTable& OwningTable, const TSet<int32>& StartingNodes, TSet<int32>* OutExternalDependencies/* = nullptr*/, TMap<int32, TArray<int32>>* OutRouteMap/* = nullptr*/)
 {
 	UE::Insights::FStopwatch Stopwatch;
 	Stopwatch.Start();
 
-	TSet<int32> AllReachableNodes = GatherAllReachableNodes(StartingNodes.Array(), OwningTable, TSet<int32>{}, TSet<const TCHAR*, TStringPointerSetKeyFuncs_DEPRECATED<const TCHAR*>>{});
+	TSet<int32> AllReachableNodes = GatherAllReachableNodes(StartingNodes.Array(), OwningTable, TSet<int32>{}, TSet<const TCHAR*, TStringPointerSetKeyFuncs_DEPRECATED<const TCHAR*>>{}, OutRouteMap);
 	TSet<const TCHAR*, TStringPointerSetKeyFuncs_DEPRECATED<const TCHAR*>> SourcePlugins;
 	for (int32 RootIndex : StartingNodes)
 	{
diff --git a/Engine/Plugins/Editor/AssetManagerEditor/Source/AssetManagerEditor/Private/TreeView/AssetTable.h b/Engine/Plugins/Editor/AssetManagerEditor/Source/AssetManagerEditor/Private/TreeView/AssetTable.h
--- a/Engine/Plugins/Editor/AssetManagerEditor/Source/AssetManagerEditor/Private/TreeView/AssetTable.h
+++ b/Engine/Plugins/Editor/AssetManagerEditor/Source/AssetManagerEditor/Private/TreeView/AssetTable.h
@@ -169,14 +169,16 @@ public:
 
 	static FAssetTableDependencySizes ComputeDependencySizes(const FAssetTable& OwningTable, const TSet<int32>& RootIndices, TSet<int32>* OutUniqueDependencies = nullptr, TSet<int32>* OutSharedDependencies = nullptr);
 
-	static int64 ComputeTotalSizeExternalDependencies(const FAssetTable& OwningTable, const TSet<int32>& StartingNodes, TSet<int32>* OutExternalDependencies = nullptr);
+	static int64 ComputeTotalSizeExternalDependencies(const FAssetTable& OwningTable, const TSet<int32>& StartingNodes, TSet<int32>* OutExternalDependencies = nullptr, TMap<int32, TArray<int32>>* OutRouteMap = nullptr);
+
+	// Finds all nodes reachable from StartingNodes INCLUDING those in StartingNodes
+	// An optional OutRouteMap may be provided. For each discovered reachable node, the routemap will get a value entry containing an array of nodes that constitute one possible path from a StartingNode to the destination
+	static TSet<int32> GatherAllReachableNodes(const TArray<int32>& StartingNodes, const FAssetTable& OwningTable, const TSet<int32>& AdditionalNodesToStopAt, const TSet<const TCHAR*, TStringPointerSetKeyFuncs_DEPRECATED<const TCHAR*>>& RestrictToPlugins, TMap<int32, TArray<int32>>* OutRouteMap = nullptr);
 
 private:
 	static void RefineDependencies(const TSet<int32>& PreviouslyVisitedIndices, const FAssetTable& OwningTable, const TSet<int32> RootIndices, const TSet<const TCHAR*, TStringPointerSetKeyFuncs_DEPRECATED<const TCHAR*>>& RestrictToGFPs, TSet<int32>& OutUniqueIndices, TSet<int32>& OutSharedIndices);
 	static inline bool ShouldSkipDueToPlugin(const TSet<const TCHAR*, TStringPointerSetKeyFuncs_DEPRECATED<const TCHAR*>>& RestrictToPlugins, const TCHAR* PluginName) { return (RestrictToPlugins.Num() > 0) && !RestrictToPlugins.Contains(PluginName); }
 
-	// Finds all nodes reachable from StartingNodes INCLUDING those in StartingNodes
-	static TSet<int32> GatherAllReachableNodes(const TArray<int32>& StartingNodes, const FAssetTable& OwningTable, const TSet<int32>& AdditionalNodesToStopAt, const TSet<const TCHAR*, TStringPointerSetKeyFuncs_DEPRECATED<const TCHAR*>>& RestrictToPlugins);
 
 private:
 	FSoftObjectPath SoftObjectPath; // Used to allow you to go to that asset in the content browser or to open its editor after lookingup the FAssetData in the original registry state
diff --git a/Engine/Plugins/Editor/AssetManagerEditor/Source/AssetManagerEditor/Private/TreeView/SAssetTableTreeView.cpp b/Engine/Plugins/Editor/AssetManagerEditor/Source/AssetManagerEditor/Private/TreeView/SAssetTableTreeView.cpp
--- a/Engine/Plugins/Editor/AssetManagerEditor/Source/AssetManagerEditor/Private/TreeView/SAssetTableTreeView.cpp
+++ b/Engine/Plugins/Editor/AssetManagerEditor/Source/AssetManagerEditor/Private/TreeView/SAssetTableTreeView.cpp
@@ -556,6 +556,33 @@ void SAssetTableTreeView::InternalCreateGroupings()
 
 ////////////////////////////////////////////////////////////////////////////////////////////////////
 
+static void WriteDependencyLine(const FAssetTable& AssetTable, const TMap<int32, TArray<int32>>& RouteMap, int32 RowIndex, FAnsiStringBuilderBase* ReusableLineBuffer, FArchive* OutDependencyFile, FString DependencyType)
+{
+	const FAssetTableRow& Row = AssetTable.GetAssetChecked(RowIndex);
+	ReusableLineBuffer->Appendf("%s%s,%s,%s", *WriteToAnsiString<512>(Row.GetPath()),
+												*WriteToAnsiString<64>(Row.GetName()), 
+												*WriteToAnsiString<32>(LexToString(Row.GetStagedCompressedSize())),
+												*WriteToAnsiString<16>(*DependencyType));
+
+	if (const TArray<int32>* Route = RouteMap.Find(RowIndex))
+	{
+		ReusableLineBuffer->AppendChar(',');
+		if (Route->Num() > 0)
+		{
+			for (int32 DependencyIndex : *Route)
+			{
+				const FAssetTableRow& DependencyRow = AssetTable.GetAssetChecked(DependencyIndex);
+				ReusableLineBuffer->Appendf("%s%s->", *WriteToAnsiString<512>(DependencyRow.GetPath()),
+					*WriteToAnsiString<64>(DependencyRow.GetName()));
+			}
+			ReusableLineBuffer->RemoveSuffix(2); // Remove the last ->
+		}
+	}
+	ReusableLineBuffer->Append(LINE_TERMINATOR);
+	OutDependencyFile->Serialize(ReusableLineBuffer->GetData(), ReusableLineBuffer->Len());
+	ReusableLineBuffer->Reset();
+}
+
 void SAssetTableTreeView::ExportDependencyData() const
 {
 	IDesktopPlatform* DesktopPlatform = FDesktopPlatformModule::Get();
@@ -593,61 +620,41 @@ void SAssetTableTreeView::ExportDependencyData() const
 		FString OutputFileName = SaveFileNames[0];
 
 
-		//FString FolderPath;
-		//const bool bFolderSelected = DesktopPlatform->OpenDirectoryDialog(
-		//	ParentWindowHandle,
-		//	LOCTEXT("FolderDialogTitle", "Choose a directory").ToString(),
-		//	FPaths::ProjectLogDir(),
-		//	FolderPath
-		//);
-
-		//if (bFolderSelected)
-		//{
-		//	if (!FolderPath.EndsWith(TEXT("/")))
-		//	{
-		//		FolderPath += TEXT("/");
-		//	}
-		//}
+		TSet<int32> ExternalDependencies;
+		TMap<int32, TArray<int32>> RouteMap;
+		FAssetTableRow::ComputeTotalSizeExternalDependencies(*GetAssetTable(), SelectedIndices, &ExternalDependencies, &RouteMap);
 
 		TSet<int32> UniqueDependencies;
 		TSet<int32> SharedDependencies;
 		FAssetTableRow::ComputeDependencySizes(*GetAssetTable(), SelectedIndices, &UniqueDependencies, &SharedDependencies);
 
-		TSet<int32> ExternalDependencies;
-		FAssetTableRow::ComputeTotalSizeExternalDependencies(*GetAssetTable(), SelectedIndices, &ExternalDependencies);
-
 		FString TimeSuffix = FDateTime::Now().ToString(TEXT("%Y%m%d_%H%M%S"));
 
 		TAnsiStringBuilder<4096> StringBuilder;
-		StringBuilder.Appendf("Asset,Self Size,Dependency Type\n");
+		TUniquePtr<FArchive> DependencyFile(IFileManager::Get().CreateFileWriter(*OutputFileName));
+
+		StringBuilder.Appendf("Asset,Self Size,Dependency Type,DependencyChain\n");
 		{
 			for (int32 RootIndex : SelectedIndices)
 			{
-				const FAssetTableRow& Row = GetAssetTable()->GetAssetChecked(RootIndex);
-				StringBuilder.Appendf("%s%s,%s,Root\n", *WriteToAnsiString<512>(Row.GetPath()), *WriteToAnsiString<64>(Row.GetName()), *WriteToAnsiString<32>(LexToString(Row.GetStagedCompressedSize())));
+				WriteDependencyLine(*GetAssetTable(), RouteMap, RootIndex, &StringBuilder, DependencyFile.Get(), TEXT("Root"));
 			}
 
 			for (int32 DependencyIndex : UniqueDependencies)
 			{
-				const FAssetTableRow& DependencyRow = GetAssetTable()->GetAssetChecked(DependencyIndex);
-				StringBuilder.Appendf("%s%s,%s,Unique\n", *WriteToAnsiString<512>(DependencyRow.GetPath()), *WriteToAnsiString<64>(DependencyRow.GetName()), *WriteToAnsiString<32>(LexToString(DependencyRow.GetStagedCompressedSize())));
+				WriteDependencyLine(*GetAssetTable(), RouteMap, DependencyIndex, &StringBuilder, DependencyFile.Get(), TEXT("Unique"));
 			}
 
 			for (int32 DependencyIndex : SharedDependencies)
 			{
-				const FAssetTableRow& DependencyRow = GetAssetTable()->GetAssetChecked(DependencyIndex);
-				StringBuilder.Appendf("%s%s,%s,Shared\n", *WriteToAnsiString<512>(DependencyRow.GetPath()), *WriteToAnsiString<64>(DependencyRow.GetName()), *WriteToAnsiString<32>(LexToString(DependencyRow.GetStagedCompressedSize())));
+				WriteDependencyLine(*GetAssetTable(), RouteMap, DependencyIndex, &StringBuilder, DependencyFile.Get(), TEXT("Shared"));
 			}
 
 			for (int32 DependencyIndex : ExternalDependencies)
 			{
-				const FAssetTableRow& DependencyRow = GetAssetTable()->GetAssetChecked(DependencyIndex);
-				StringBuilder.Appendf("%s%s,%s,External\n", *WriteToAnsiString<512>(DependencyRow.GetPath()), *WriteToAnsiString<64>(DependencyRow.GetName()), *WriteToAnsiString<32>(LexToString(DependencyRow.GetStagedCompressedSize())));
+				WriteDependencyLine(*GetAssetTable(), RouteMap, DependencyIndex, &StringBuilder, DependencyFile.Get(), TEXT("External"));
 			}
 		}
-
-		TUniquePtr<FArchive> DependencyFile(IFileManager::Get().CreateFileWriter(*OutputFileName));
-		DependencyFile->Serialize(StringBuilder.GetData(), StringBuilder.Len());
 	}
 }
 
@@ -817,6 +824,56 @@ FText SAssetTableTreeView::GetFooterRightText1() const
 
 ////////////////////////////////////////////////////////////////////////////////////////////////////
 
+void SAssetTableTreeView::CalculateBaseAndMarginalCostForSelection(TSet<int32>& SelectionSetIndices, int64* OutTotalSizeMultiplyUsed, int64* OutTotalSizeSingleUse) const
+{
+	if (!ensure((OutTotalSizeMultiplyUsed != nullptr) && (OutTotalSizeSingleUse != nullptr)))
+	{
+		return;
+	}
+
+	TSet<const TCHAR*, TStringPointerSetKeyFuncs_DEPRECATED<const TCHAR*>> SelectedPlugins;
+	for (int32 ItemIndex : SelectionSetIndices)
+	{
+		SelectedPlugins.Add(GetAssetTable()->GetAssetChecked(ItemIndex).GetPluginName());
+	}
+
+	TMap<int32, int32> NumItemsReferencingDependency;
+	for (int32 ItemRowIndex : SelectionSetIndices)
+	{
+		TSet<int32> Dependencies = FAssetTableRow::GatherAllReachableNodes(TArray<int32>{ItemRowIndex}, * GetAssetTable(), TSet<int32>{}, SelectedPlugins);
+		for (int32 DependencyRowIndex : Dependencies)
+		{
+			if (SelectionSetIndices.Contains(DependencyRowIndex))
+			{
+				// Don't count assets we've selected, we'll handle those separately
+				continue;
+			}
+			if (int32* Count = NumItemsReferencingDependency.Find(DependencyRowIndex))
+			{
+				(*Count)++;
+			}
+			else
+			{
+				NumItemsReferencingDependency.Add(DependencyRowIndex, 1);
+			}
+		}
+	}
+
+	for (TPair<int32, int32> ReferenceCountPair : NumItemsReferencingDependency)
+	{
+		if (ReferenceCountPair.Value > 1)
+		{
+			OutTotalSizeMultiplyUsed += GetAssetTable()->GetAssetChecked(ReferenceCountPair.Key).GetStagedCompressedSize();
+		}
+		else if (ReferenceCountPair.Value == 1)
+		{
+			OutTotalSizeSingleUse += GetAssetTable()->GetAssetChecked(ReferenceCountPair.Key).GetStagedCompressedSize();
+		}
+	}
+}
+
+////////////////////////////////////////////////////////////////////////////////////////////////////
+
 void SAssetTableTreeView::TreeView_OnSelectionChanged(UE::Insights::FTableTreeNodePtr SelectedItem, ESelectInfo::Type SelectInfo)
 {
 	TArray<UE::Insights::FTableTreeNodePtr> SelectedNodes;
@@ -825,7 +882,7 @@ void SAssetTableTreeView::TreeView_OnSelectionChanged(UE::Insights::FTableTreeNo
 	FAssetTreeNodePtr NewSelectedAssetNode;
 	int32 NewlySelectedAssetRowIndex = -1;
 	TSet<int32> SelectionSetIndices;
-
+	
 	for (const UE::Insights::FTableTreeNodePtr& Node : SelectedNodes)
 	{
 		if (Node->Is<FAssetTreeNode>() && Node->As<FAssetTreeNode>().IsValidAsset())
@@ -875,16 +932,30 @@ void SAssetTableTreeView::TreeView_OnSelectionChanged(UE::Insights::FTableTreeNo
 	}
 	else
 	{
-		if (FilteredAssetCount != VisibleAssetCount)
-		{
-			FooterLeftText = FText::Format(LOCTEXT("FooterLeftTextFmt_ManySelected_Filtered", "{0} / {1} assets ({2} selected)"), FText::AsNumber(FilteredAssetCount), FText::AsNumber(VisibleAssetCount), FText::AsNumber(NumSelectedAssets));
-		}
-		else
+		bool AllSelectedNodesAreSameType = true;
 		{
-			FooterLeftText = FText::Format(LOCTEXT("FooterLeftTextFmt_ManySelected_NoFiltered", "{0} assets ({1} selected)"), FText::AsNumber(VisibleAssetCount), FText::AsNumber(NumSelectedAssets));
+			const TCHAR* FirstNativeClass = nullptr;
+			bool FirstIndex = true;
+			for (int32 SelectedNodeIndex : SelectionSetIndices)
+			{
+				const FAssetTableRow& AssetTableRow = GetAssetTable()->GetAssetChecked(SelectedNodeIndex);
+				const TCHAR* NativeClass = AssetTableRow.GetNativeClass();
+				if (FirstIndex)
+				{
+					FirstNativeClass = NativeClass;
+					FirstIndex = false;
+				}
+				else if (NativeClass != FirstNativeClass)
+				{
+					AllSelectedNodesAreSameType = false;
+					break;
+				}
+			}
 		}
-		FooterCenterText1 = FText();
-		FooterCenterText2 = FText();
+
+
+		int64 TotalExternalDependencySize = FAssetTableRow::ComputeTotalSizeExternalDependencies(*GetAssetTable(), SelectionSetIndices);
+		FAssetTableDependencySizes Sizes = FAssetTableRow::ComputeDependencySizes(*GetAssetTable(), SelectionSetIndices, nullptr, nullptr);
 
 		int64 TotalSelfSize = 0;
 		for (int32 Index : SelectionSetIndices)
@@ -892,8 +963,33 @@ void SAssetTableTreeView::TreeView_OnSelectionChanged(UE::Insights::FTableTreeNo
 			TotalSelfSize += GetAssetTable()->GetAssetChecked(Index).GetStagedCompressedSize();
 		}
 
-		int64 TotalExternalDependencySize = FAssetTableRow::ComputeTotalSizeExternalDependencies(*GetAssetTable(), SelectionSetIndices);
-		FAssetTableDependencySizes Sizes = FAssetTableRow::ComputeDependencySizes(*GetAssetTable(), SelectionSetIndices, nullptr, nullptr);
+		FText BaseAndMarginalCost;
+		if (AllSelectedNodesAreSameType)
+		{
+			int64 TotalSizeMultiplyUsed = 0;
+			int64 TotalSizeSingleUse = 0;
+			CalculateBaseAndMarginalCostForSelection(SelectionSetIndices, &TotalSizeMultiplyUsed, &TotalSizeSingleUse);
+
+			BaseAndMarginalCost = FText::Format(LOCTEXT("FooterLeft_BaseAndMarginalCost", " -- Base Cost: {0}  Per Asset Cost: {1}"),
+				FText::AsMemory(TotalSizeMultiplyUsed),
+				FText::AsMemory((TotalSelfSize + TotalSizeSingleUse) / SelectionSetIndices.Num()));
+		}
+		else
+		{
+			BaseAndMarginalCost = LOCTEXT("FooterLeft_BaseAndMarginalCost_MultipleTypesError", " -- Multiple types selected");
+		}
+
+		if (FilteredAssetCount != VisibleAssetCount)
+		{
+			FooterLeftText = FText::Format(LOCTEXT("FooterLeftTextFmt_ManySelected_Filtered", "{0} / {1} assets ({2} selected{3})"), FText::AsNumber(FilteredAssetCount), FText::AsNumber(VisibleAssetCount), FText::AsNumber(NumSelectedAssets), BaseAndMarginalCost);
+		}
+		else
+		{
+			FooterLeftText = FText::Format(LOCTEXT("FooterLeftTextFmt_ManySelected_NoFiltered", "{0} assets ({1} selected{2})"), FText::AsNumber(VisibleAssetCount), FText::AsNumber(NumSelectedAssets), BaseAndMarginalCost);
+		}
+		FooterCenterText1 = FText();
+		FooterCenterText2 = FText();
+
 		FooterRightText1 = FText::Format(LOCTEXT("FooterRightFmt", "Self: {0}    Unique: {1}    Shared: {2}    External: {3}"),
 			FText::AsMemory(TotalSelfSize),
 			FText::AsMemory(Sizes.UniqueDependenciesSize),
diff --git a/Engine/Plugins/Editor/AssetManagerEditor/Source/AssetManagerEditor/Private/TreeView/SAssetTableTreeView.h b/Engine/Plugins/Editor/AssetManagerEditor/Source/AssetManagerEditor/Private/TreeView/SAssetTableTreeView.h
--- a/Engine/Plugins/Editor/AssetManagerEditor/Source/AssetManagerEditor/Private/TreeView/SAssetTableTreeView.h
+++ b/Engine/Plugins/Editor/AssetManagerEditor/Source/AssetManagerEditor/Private/TreeView/SAssetTableTreeView.h
@@ -28,6 +28,7 @@ public:
 	{
 	}
 	SLATE_EVENT(FOnSelectionChanged, OnSelectionChanged)
+
 	SLATE_END_ARGS()
 
 	/**
@@ -90,6 +91,8 @@ private:
 
 	TArray<FAssetData> GetAssetDataForSelection() const;
 
+	void CalculateBaseAndMarginalCostForSelection(TSet<int32>& SelectionSetIndices, int64* OutTotalSizeMultiplyUsed, int64* OutTotalSizeSingleUse) const;
+
 private:
 	bool bNeedsToRebuild = false;
 
