AssetManagerEditor: Separated the new Audit / Asset Disk Size tool from the existing Asset Audit tool. The disk size tool now maintains its own asset registry source independent from the state of the main asset registry editor module.

diff --git a/Engine/Plugins/Editor/AssetManagerEditor/Source/AssetManagerEditor/Private/AssetManagerEditorModule.cpp b/Engine/Plugins/Editor/AssetManagerEditor/Source/AssetManagerEditor/Private/AssetManagerEditorModule.cpp
--- a/Engine/Plugins/Editor/AssetManagerEditor/Source/AssetManagerEditor/Private/AssetManagerEditorModule.cpp
+++ b/Engine/Plugins/Editor/AssetManagerEditor/Source/AssetManagerEditor/Private/AssetManagerEditorModule.cpp
@@ -61,6 +61,8 @@
 #include "ContentBrowserDataSubsystem.h"
 #include "Insights/Common/InsightsStyle.h"
 #include "Insights/Filter/ViewModels/Filters.h"
+#include "TreeView/AssetTable.h"
+#include "TreeView/SAssetTableTreeView.h"
 
 #define LOCTEXT_NAMESPACE "AssetManagerEditor"
 
@@ -313,20 +315,19 @@ public:
 	virtual void OpenSizeMapUI(TArray<FAssetIdentifier> SelectedIdentifiers) override;
 	virtual void OpenSizeMapUI(TArray<FName> SelectedPackages) override;	
 	virtual void OpenShaderCookStatistics(TArray<FName> SelectedIdentifiers) override;
-	virtual bool GetStringValueForCustomColumn(const FAssetData& AssetData, FName ColumnName, FString& OutValue) override;
-	virtual bool GetDisplayTextForCustomColumn(const FAssetData& AssetData, FName ColumnName, FText& OutValue) override;
-	virtual bool GetIntegerValueForCustomColumn(const FAssetData& AssetData, FName ColumnName, int64& OutValue) override;
+	virtual bool GetStringValueForCustomColumn(const FAssetData& AssetData, FName ColumnName, FString& OutValue, const FAssetManagerEditorRegistrySource* OverrideRegistrySource = nullptr) override;
+	virtual bool GetDisplayTextForCustomColumn(const FAssetData& AssetData, FName ColumnName, FText& OutValue, const FAssetManagerEditorRegistrySource* OverrideRegistrySource = nullptr) override;
+	virtual bool GetIntegerValueForCustomColumn(const FAssetData& AssetData, FName ColumnName, int64& OutValue, const FAssetManagerEditorRegistrySource* OverrideRegistrySource = nullptr) override;
 	virtual bool GetManagedPackageListForAssetData(const FAssetData& AssetData, TSet<FName>& ManagedPackageSet) override;
 	virtual void GetAvailableRegistrySources(TArray<const FAssetManagerEditorRegistrySource*>& AvailableSources) override;
 	virtual const FAssetManagerEditorRegistrySource* GetCurrentRegistrySource(bool bNeedManagementData = false) override;
 	virtual void SetCurrentRegistrySource(const FString& SourceName) override;
+	virtual bool PopulateRegistrySource(FAssetManagerEditorRegistrySource* OutRegistrySource) override;
 	virtual void RefreshRegistryData() override;
 	virtual bool IsPackageInCurrentRegistrySource(FName PackageName) override;
 	virtual bool FilterAssetIdentifiersForCurrentRegistrySource(TArray<FAssetIdentifier>& AssetIdentifiers, const FAssetManagerDependencyQuery& DependencyQuery = FAssetManagerDependencyQuery::None(), bool bForwardDependency = true) override;
 	virtual bool WriteCollection(FName CollectionName, ECollectionShareType::Type ShareType, const TArray<FName>& PackageNames, bool bShowFeedback) override;
 
-	virtual bool ShouldEnableTreeViewInAssetAudit() const override;
-
 private:
 
 	static bool GetDependencyTypeArg(const FString& Arg, UE::AssetRegistry::EDependencyQuery& OutRequiredFlags);
@@ -360,6 +361,7 @@ private:
 	static const TCHAR* FindDepChainHelpText;
 	static const TCHAR* FindClassDepHelpText;
 	static const FName AssetAuditTabName;
+	static const FName AssetDiskSizeTabName;
 	static const FName ReferenceViewerTabName;
 	static const FName SizeMapTabName;
 
@@ -368,9 +370,11 @@ private:
 	FDelegateHandle AssetEditorExtenderDelegateHandle;
 
 	TWeakPtr<SDockTab> AssetAuditTab;
+	TWeakPtr<SDockTab> AssetDiskSizeTab;
 	TWeakPtr<SDockTab> ReferenceViewerTab;
 	TWeakPtr<SDockTab> SizeMapTab;
 	TWeakPtr<SAssetAuditBrowser> AssetAuditUI;
+	TWeakPtr<SAssetTableTreeView> AssetDiskSizeUI;
 	TWeakPtr<SSizeMap> SizeMapUI;
 	TWeakPtr<SReferenceViewer> ReferenceViewerUI;
 	TMap<FString, FAssetManagerEditorRegistrySource> RegistrySourceMap;
@@ -400,15 +404,15 @@ private:
 	void OnEditAssetIdentifiers(TArray<FAssetIdentifier> AssetIdentifiers);
 
 	TSharedRef<SDockTab> SpawnAssetAuditTab(const FSpawnTabArgs& Args);
+	TSharedRef<SDockTab> SpawnAssetDiskSizeTab(const FSpawnTabArgs& Args);
 	TSharedRef<SDockTab> SpawnReferenceViewerTab(const FSpawnTabArgs& Args);
 	TSharedRef<SDockTab> SpawnSizeMapTab(const FSpawnTabArgs& Args);
-
-	IConsoleVariable* EnableTreeViewInAssetAuditCVar = nullptr;
 };
 
 const TCHAR* FAssetManagerEditorModule::FindDepChainHelpText = TEXT("Finds all dependency chains from assets in the given search path, to the target package.\n Usage: FindDepChain TargetPackagePath SearchRootPath (optional: -hardonly/-softonly)\n e.g. FindDepChain /game/characters/heroes/muriel/meshes/muriel /game/cards ");
 const TCHAR* FAssetManagerEditorModule::FindClassDepHelpText = TEXT("Finds all dependencies of a certain set of classes to the target asset.\n Usage: FindDepClasses TargetPackagePath ClassName1 ClassName2 etc (optional: -hardonly/-softonly) \n e.g. FindDepChain /game/characters/heroes/muriel/meshes/muriel /game/cards");
 const FName FAssetManagerEditorModule::AssetAuditTabName = TEXT("AssetAudit");
+const FName FAssetManagerEditorModule::AssetDiskSizeTabName = TEXT("AssetDiskSize");
 const FName FAssetManagerEditorModule::ReferenceViewerTabName = TEXT("ReferenceViewer");
 const FName FAssetManagerEditorModule::SizeMapTabName = TEXT("SizeMap");
 
@@ -436,12 +440,6 @@ void FAssetManagerEditorModule::StartupModule()
 		UE::Insights::FInsightsStyle::Initialize();
 		UE::Insights::FFilterService::Initialize();
 
-		EnableTreeViewInAssetAuditCVar = IConsoleManager::Get().RegisterConsoleVariable(
-			TEXT("AssetAudit.EnableTreeView"),
-			false,
-			TEXT("Enables the tree view in Asset Audit (next time the window is open)."),
-			ECVF_Default);
-
 		AuditCmds.Add(IConsoleManager::Get().RegisterConsoleCommand(
 			TEXT("AssetManager.AssetAudit"),
 			TEXT("Dumps statistics about assets to the log."),
@@ -519,6 +517,13 @@ void FAssetManagerEditorModule::StartupModule()
 			.SetIcon(FSlateIcon(FAppStyle::GetAppStyleSetName(), "LevelEditor.Audit"));
 		FGlobalTabmanager::Get()->RegisterDefaultTabWindowSize(AssetAuditTabName, FVector2D(1080, 600));
 
+		FGlobalTabmanager::Get()->RegisterNomadTabSpawner(AssetDiskSizeTabName, FOnSpawnTab::CreateRaw(this, &FAssetManagerEditorModule::SpawnAssetDiskSizeTab))
+			.SetDisplayName(LOCTEXT("AssetDiskSizeTitle", "Asset Disk Size"))
+			.SetTooltipText(LOCTEXT("AssetDiskSizeTooltip", "Open Asset Disk Size window, allows to analyze disk size for assets."))
+			.SetGroup(WorkspaceMenu::GetMenuStructure().GetDeveloperToolsAuditCategory())
+			.SetIcon(FSlateIcon(FAppStyle::GetAppStyleSetName(), "LevelEditor.Audit"));
+		FGlobalTabmanager::Get()->RegisterDefaultTabWindowSize(AssetDiskSizeTabName, FVector2D(1080, 600));
+
 		FGlobalTabmanager::Get()->RegisterNomadTabSpawner(ReferenceViewerTabName, FOnSpawnTab::CreateRaw(this, &FAssetManagerEditorModule::SpawnReferenceViewerTab))
 			.SetDisplayName(LOCTEXT("ReferenceViewerTitle", "Reference Viewer"))
 			.SetMenuType(ETabSpawnerMenuType::Hidden);
@@ -555,12 +560,6 @@ void FAssetManagerEditorModule::ShutdownModule()
 	}
 	AuditCmds.Empty();
 
-	if (EnableTreeViewInAssetAuditCVar)
-	{
-		IConsoleManager::Get().UnregisterConsoleObject(EnableTreeViewInAssetAuditCVar);
-		EnableTreeViewInAssetAuditCVar = nullptr;
-	}
-
 	if ((GIsEditor && !IsRunningCommandlet()) && UObjectInitialized() && FSlateApplication::IsInitialized())
 	{
 		FContentBrowserModule& ContentBrowserModule = FModuleManager::LoadModuleChecked<FContentBrowserModule>(TEXT("ContentBrowser"));
@@ -595,6 +594,7 @@ void FAssetManagerEditorModule::ShutdownModule()
 		}
 
 		FGlobalTabmanager::Get()->UnregisterNomadTabSpawner(AssetAuditTabName);
+		FGlobalTabmanager::Get()->UnregisterNomadTabSpawner(AssetDiskSizeTabName);
 		FGlobalTabmanager::Get()->UnregisterNomadTabSpawner(ReferenceViewerTabName);
 		FGlobalTabmanager::Get()->UnregisterNomadTabSpawner(SizeMapTabName);
 
@@ -649,6 +649,31 @@ TSharedRef<SDockTab> FAssetManagerEditorModule::SpawnAssetAuditTab(const FSpawnT
 	return DockTab;
 }
 
+TSharedRef<SDockTab> FAssetManagerEditorModule::SpawnAssetDiskSizeTab(const FSpawnTabArgs& Args)
+{
+	check(UAssetManager::IsInitialized());
+
+	TSharedRef<FAssetTable> AssetTable = MakeShared<FAssetTable>();
+	AssetTable->Reset();
+	AssetTable->SetDisplayName(FText::FromString(TEXT("AssetTable")));
+
+	TSharedRef<SDockTab> DockTab = SAssignNew(AssetDiskSizeTab, SDockTab)
+		.TabRole(ETabRole::NomadTab)
+		[
+			SAssignNew(AssetDiskSizeUI, SAssetTableTreeView, AssetTable)
+		];
+
+	DockTab->SetOnTabClosed(SDockTab::FOnTabClosedCallback::CreateLambda([this](TSharedRef<SDockTab>)
+		{
+			if (AssetDiskSizeUI.IsValid())
+			{
+				AssetDiskSizeUI.Pin()->OnClose();
+			}
+		}));
+
+	return DockTab;
+}
+
 TSharedRef<SDockTab> FAssetManagerEditorModule::SpawnReferenceViewerTab(const FSpawnTabArgs& Args)
 {
 	TSharedRef<SDockTab> NewTab = SAssignNew(ReferenceViewerTab, SDockTab)
@@ -1235,9 +1260,11 @@ bool FAssetManagerEditorModule::GetManagedPackageListForAssetData(const FAssetDa
 	return bFoundAny;
 }
 
-bool FAssetManagerEditorModule::GetStringValueForCustomColumn(const FAssetData& AssetData, FName ColumnName, FString& OutValue)
+bool FAssetManagerEditorModule::GetStringValueForCustomColumn(const FAssetData& AssetData, FName ColumnName, FString& OutValue, const FAssetManagerEditorRegistrySource* OverrideRegistrySource /*= nullptr*/)
 {
-	if (!CurrentRegistrySource || !CurrentRegistrySource->HasRegistry())
+	const FAssetManagerEditorRegistrySource* RegistrySource = (OverrideRegistrySource == nullptr) ? CurrentRegistrySource : OverrideRegistrySource;
+
+	if (!RegistrySource || !RegistrySource->HasRegistry())
 	{
 		return false;
 	}
@@ -1284,14 +1311,14 @@ bool FAssetManagerEditorModule::GetStringValueForCustomColumn(const FAssetData&
 		TArray<int32> FoundChunks;
 		OutValue.Reset();
 
-		if (CurrentRegistrySource->bIsEditor)
+		if (RegistrySource->bIsEditor)
 		{
 			// The in-memory data is wrong, ask the asset manager
-			AssetManager.GetPackageChunkIds(AssetData.PackageName, CurrentRegistrySource->TargetPlatform, AssetData.GetChunkIDs(), FoundChunks);
+			AssetManager.GetPackageChunkIds(AssetData.PackageName, RegistrySource->TargetPlatform, AssetData.GetChunkIDs(), FoundChunks);
 		}
 		else
 		{
-			FAssetData PlatformData = CurrentRegistrySource->GetAssetByObjectPath(AssetData.GetSoftObjectPath());
+			FAssetData PlatformData = RegistrySource->GetAssetByObjectPath(AssetData.GetSoftObjectPath());
 			if (PlatformData.IsValid())
 			{
 				FoundChunks = PlatformData.GetChunkIDs();
@@ -1328,15 +1355,17 @@ bool FAssetManagerEditorModule::GetStringValueForCustomColumn(const FAssetData&
 	else
 	{
 		// Get base value of asset tag
-		return CurrentRegistrySource->GetAssetTagByObjectPath(AssetData.GetSoftObjectPath(), ColumnName, OutValue);
+		return RegistrySource->GetAssetTagByObjectPath(AssetData.GetSoftObjectPath(), ColumnName, OutValue);
 	}
 
 	return false;
 }
 
-bool FAssetManagerEditorModule::GetDisplayTextForCustomColumn(const FAssetData& AssetData, FName ColumnName, FText& OutValue)
+bool FAssetManagerEditorModule::GetDisplayTextForCustomColumn(const FAssetData& AssetData, FName ColumnName, FText& OutValue, const FAssetManagerEditorRegistrySource* OverrideRegistrySource /*= nullptr*/)
 {
-	if (!CurrentRegistrySource || !CurrentRegistrySource->HasRegistry())
+	const FAssetManagerEditorRegistrySource* RegistrySource = (OverrideRegistrySource == nullptr) ? CurrentRegistrySource : OverrideRegistrySource;
+
+	if (!RegistrySource || !RegistrySource->HasRegistry())
 	{
 		return false;
 	}
@@ -1399,15 +1428,17 @@ bool FAssetManagerEditorModule::GetDisplayTextForCustomColumn(const FAssetData&
 	else
 	{
 		// Get base value of asset tag
-		return CurrentRegistrySource->GetAssetTagByObjectPath(AssetData.GetSoftObjectPath(), ColumnName, OutValue);
+		return RegistrySource->GetAssetTagByObjectPath(AssetData.GetSoftObjectPath(), ColumnName, OutValue);
 	}
 
 	return false;
 }
 
-bool FAssetManagerEditorModule::GetIntegerValueForCustomColumn(const FAssetData& AssetData, FName ColumnName, int64& OutValue)
+bool FAssetManagerEditorModule::GetIntegerValueForCustomColumn(const FAssetData& AssetData, FName ColumnName, int64& OutValue, const FAssetManagerEditorRegistrySource* OverrideRegistrySource /*= nullptr*/)
 {
-	if (!CurrentRegistrySource || !CurrentRegistrySource->HasRegistry())
+	const FAssetManagerEditorRegistrySource* RegistrySource = (OverrideRegistrySource == nullptr) ? CurrentRegistrySource : OverrideRegistrySource;
+
+	if (!RegistrySource || !RegistrySource->HasRegistry())
 	{
 		return false;
 	}
@@ -1459,13 +1490,13 @@ bool FAssetManagerEditorModule::GetIntegerValueForCustomColumn(const FAssetData&
 	{
 		// If the asset registry has the staged sizes written back, then use that as an accurate representation of the
 		// actual chunks in the package.
-		if (CurrentRegistrySource->GetAssetTagByObjectPath(AssetData.GetSoftObjectPath(), StageChunkSizeName, OutValue))
+		if (RegistrySource->GetAssetTagByObjectPath(AssetData.GetSoftObjectPath(), StageChunkSizeName, OutValue))
 		{
 			return true;
 		}
 
 		// Otherwise, use the DiskSize generated by SavePackage.
-		TOptional<FAssetPackageData> FoundData = CurrentRegistrySource->GetAssetPackageDataCopy(AssetData.PackageName);
+		TOptional<FAssetPackageData> FoundData = RegistrySource->GetAssetPackageDataCopy(AssetData.PackageName);
 
 		if (FoundData && FoundData->DiskSize >= 0)
 		{
@@ -1508,7 +1539,7 @@ bool FAssetManagerEditorModule::GetIntegerValueForCustomColumn(const FAssetData&
 	else
 	{
 		// Get base value of asset tag
-		return CurrentRegistrySource->GetAssetTagByObjectPath(AssetData.GetSoftObjectPath(), ColumnName, OutValue);
+		return RegistrySource->GetAssetTagByObjectPath(AssetData.GetSoftObjectPath(), ColumnName, OutValue);
 	}
 
 	return false;
@@ -1622,100 +1653,91 @@ void FAssetManagerEditorRegistrySource::LoadRegistryTimestamp()
 	}
 }
 
-void FAssetManagerEditorModule::SetCurrentRegistrySource(const FString& SourceName)
+bool FAssetManagerEditorModule::PopulateRegistrySource(FAssetManagerEditorRegistrySource* InOutRegistrySource)
 {
-	InitializeRegistrySources(false);
-
-	FAssetManagerEditorRegistrySource* NewSource = RegistrySourceMap.Find(SourceName);
-
-	if (NewSource)
+	if (InOutRegistrySource->SourceName == FAssetManagerEditorRegistrySource::CustomSourceName)
 	{
-		CurrentRegistrySource = NewSource;
+		IDesktopPlatform* DesktopPlatform = FDesktopPlatformModule::Get();
+		const void* ParentWindowWindowHandle = FSlateApplication::Get().FindBestParentWindowHandleForDialogs(nullptr);
+		const TCHAR* DevelopmentAssetRegistryFilename = GetDevelopmentAssetRegistryFilename();
+		const FText Title = LOCTEXT("LoadAssetRegistry", "Load DevelopmentAssetRegistry");
+		const FString FileTypes = FString::Printf(TEXT("%s|*.bin"), DevelopmentAssetRegistryFilename);
+
+		TArray<FString> OutFilenames;
+		DesktopPlatform->OpenFileDialog(
+			ParentWindowWindowHandle,
+			Title.ToString(),
+			TEXT(""),
+			DevelopmentAssetRegistryFilename,
+			FileTypes,
+			EFileDialogFlags::None,
+			OutFilenames
+		);
 
-		if (CurrentRegistrySource->SourceName == FAssetManagerEditorRegistrySource::CustomSourceName)
+		if (OutFilenames.Num() == 1)
 		{
-			CurrentRegistrySource->SourceTimestamp = FString();
+			InOutRegistrySource->SourceTimestamp = FString();
 
-			if (CurrentRegistrySource->HasRegistry())
+			if (InOutRegistrySource->HasRegistry())
 			{
-				CurrentRegistrySource->ClearRegistry();
-				CurrentRegistrySource->ChunkAssignments.Reset();
-				CurrentRegistrySource->bManagementDataInitialized = false;
+				InOutRegistrySource->ClearRegistry();
+				InOutRegistrySource->ChunkAssignments.Reset();
+				InOutRegistrySource->bManagementDataInitialized = false;
 			}
-			CurrentRegistrySource->SourceFilename.Reset();
-
-			IDesktopPlatform* DesktopPlatform = FDesktopPlatformModule::Get();
-			const void* ParentWindowWindowHandle = FSlateApplication::Get().FindBestParentWindowHandleForDialogs(nullptr);
-			const TCHAR* DevelopmentAssetRegistryFilename = GetDevelopmentAssetRegistryFilename();
-			const FText Title = LOCTEXT("LoadAssetRegistry", "Load DevelopmentAssetRegistry");
-			const FString FileTypes = FString::Printf(TEXT("%s|*.bin"), DevelopmentAssetRegistryFilename);
-
-			TArray<FString> OutFilenames;
-			DesktopPlatform->OpenFileDialog(
-				ParentWindowWindowHandle,
-				Title.ToString(),
-				TEXT(""),
-				DevelopmentAssetRegistryFilename,
-				FileTypes,
-				EFileDialogFlags::None,
-				OutFilenames
-			);
-
-			if (OutFilenames.Num() == 1)
-			{
-				CurrentRegistrySource->SourceFilename = OutFilenames[0];
-			}
-			else
-			{
-				CurrentRegistrySource = RegistrySourceMap.Find(FAssetManagerEditorRegistrySource::EditorSourceName);
-			}
-		}
 
-		if (!CurrentRegistrySource->HasRegistry() && !CurrentRegistrySource->SourceFilename.IsEmpty() && !CurrentRegistrySource->bIsEditor)
+			InOutRegistrySource->SourceFilename = OutFilenames[0];
+		}
+		else
 		{
-			CurrentRegistrySource->SourceTimestamp = FString();
+			return false;
+		}
+	}
 
-			bool bLoaded = false;
-			FArrayReader SerializedAssetData;
-			if (FFileHelper::LoadFileToArray(SerializedAssetData, *CurrentRegistrySource->SourceFilename))
-			{
-				CurrentRegistrySource->LoadRegistryTimestamp();
+	if (!InOutRegistrySource->HasRegistry() && !InOutRegistrySource->SourceFilename.IsEmpty() && !InOutRegistrySource->bIsEditor)
+	{
+		InOutRegistrySource->SourceTimestamp = FString();
 
-				FAssetRegistryState* NewState = new FAssetRegistryState();
-				FAssetRegistrySerializationOptions Options(UE::AssetRegistry::ESerializationTarget::ForDevelopment);
+		bool bLoaded = false;
+		FArrayReader SerializedAssetData;
+		if (FFileHelper::LoadFileToArray(SerializedAssetData, *InOutRegistrySource->SourceFilename))
+		{
+			InOutRegistrySource->LoadRegistryTimestamp();
 
-				NewState->Serialize(SerializedAssetData, Options);
+			FAssetRegistryState* NewState = new FAssetRegistryState();
+			FAssetRegistrySerializationOptions Options(UE::AssetRegistry::ESerializationTarget::ForDevelopment);
 
-				if (NewState->GetNumAssets() > 0)
-				{
-					bLoaded = true;
-					CurrentRegistrySource->SetRegistryState(NewState);
-				}
-			}
-			if (!bLoaded)
+			NewState->Serialize(SerializedAssetData, Options);
+
+			if (NewState->GetNumAssets() > 0)
 			{
-				FNotificationInfo Info(FText::Format(LOCTEXT("LoadRegistryFailed_FailedToLoad", "Failed to load asset registry from {0}!"), FText::FromString(CurrentRegistrySource->SourceFilename)));
-				Info.ExpireDuration = 10.0f;
-				FSlateNotificationManager::Get().AddNotification(Info);
-				CurrentRegistrySource = RegistrySourceMap.Find(FAssetManagerEditorRegistrySource::EditorSourceName);
+				bLoaded = true;
+				InOutRegistrySource->SetRegistryState(NewState);
 			}
 		}
+		if (!bLoaded)
+		{
+			FNotificationInfo Info(FText::Format(LOCTEXT("LoadRegistryFailed_FailedToLoad", "Failed to load asset registry from {0}!"), FText::FromString(InOutRegistrySource->SourceFilename)));
+			Info.ExpireDuration = 10.0f;
+			FSlateNotificationManager::Get().AddNotification(Info);
+			return false;
+		}
+	}
 
-		if (!CurrentRegistrySource->bManagementDataInitialized && CurrentRegistrySource->HasRegistry())
+	if (!InOutRegistrySource->bManagementDataInitialized && InOutRegistrySource->HasRegistry())
+	{
+		InOutRegistrySource->ChunkAssignments.Reset();
+		if (InOutRegistrySource->bIsEditor)
 		{
-			CurrentRegistrySource->ChunkAssignments.Reset();
-			if (CurrentRegistrySource->bIsEditor)
-			{
-				// Load chunk list from asset manager
-				CurrentRegistrySource->ChunkAssignments = UAssetManager::Get().GetChunkManagementMap();
-			}
-			else
-			{
-				// Iterate assets and look for chunks
-				const FAssetRegistryState* RegistryState = CurrentRegistrySource->GetOwnedRegistryState();
-				checkf(RegistryState, TEXT("Should be non-null because HasRegistry() && !bIsEditor"));
+			// Load chunk list from asset manager
+			InOutRegistrySource->ChunkAssignments = UAssetManager::Get().GetChunkManagementMap();
+		}
+		else
+		{
+			// Iterate assets and look for chunks
+			const FAssetRegistryState* RegistryState = InOutRegistrySource->GetOwnedRegistryState();
+			checkf(RegistryState, TEXT("Should be non-null because HasRegistry() && !bIsEditor"));
 
-				RegistryState->EnumerateAllAssets([&RegistryState, this](const FAssetData& AssetData)
+			RegistryState->EnumerateAllAssets([&RegistryState, &InOutRegistrySource, this](const FAssetData& AssetData)
 				{
 					const FAssetData::FChunkArrayView ChunkIDs = AssetData.GetChunkIDs();
 					if (!ChunkIDs.IsEmpty())
@@ -1726,14 +1748,14 @@ void FAssetManagerEditorModule::SetCurrentRegistrySource(const FString& SourceNa
 						for (int32 ChunkId : ChunkIDs)
 						{
 							FPrimaryAssetId ChunkAssetId = UAssetManager::CreatePrimaryAssetIdFromChunkId(ChunkId);
-							
-							FAssetManagerChunkInfo* ChunkAssignmentSet = CurrentRegistrySource->ChunkAssignments.Find(ChunkId);
+
+							FAssetManagerChunkInfo* ChunkAssignmentSet = InOutRegistrySource->ChunkAssignments.Find(ChunkId);
 
 							if (!ChunkAssignmentSet)
 							{
 								// First time found, read the graph
-								ChunkAssignmentSet = &CurrentRegistrySource->ChunkAssignments.Add(ChunkId);
-								
+								ChunkAssignmentSet = &InOutRegistrySource->ChunkAssignments.Add(ChunkId);
+
 								TArray<FAssetIdentifier> ManagedAssets;
 
 								RegistryState->GetDependencies(ChunkAssetId, ManagedAssets, UE::AssetRegistry::EDependencyCategory::Manage);
@@ -1763,10 +1785,28 @@ void FAssetManagerEditorModule::SetCurrentRegistrySource(const FString& SourceNa
 						}
 					}
 				});
-			}
+		}
+
+		InOutRegistrySource->bManagementDataInitialized = true;
+	}
+
+	return true;
+}
+
+void FAssetManagerEditorModule::SetCurrentRegistrySource(const FString& SourceName)
+{
+	InitializeRegistrySources(false);
+
+	FAssetManagerEditorRegistrySource* NewSource = RegistrySourceMap.Find(SourceName);
 
-			CurrentRegistrySource->bManagementDataInitialized = true;
+	if (NewSource)
+	{
+		if (!PopulateRegistrySource(NewSource))
+		{
+			bool Success = PopulateRegistrySource(RegistrySourceMap.Find(FAssetManagerEditorRegistrySource::EditorSourceName));
+			ensure(Success);
 		}
+		CurrentRegistrySource = NewSource;
 	}
 	else
 	{
@@ -2464,9 +2504,4 @@ PRAGMA_ENABLE_DEPRECATION_WARNINGS
 	return bSuccess;
 }
 
-bool FAssetManagerEditorModule::ShouldEnableTreeViewInAssetAudit() const
-{
-	return EnableTreeViewInAssetAuditCVar && EnableTreeViewInAssetAuditCVar->GetBool();
-}
-
 #undef LOCTEXT_NAMESPACE
diff --git a/Engine/Plugins/Editor/AssetManagerEditor/Source/AssetManagerEditor/Private/Insights/Table/Widgets/STableTreeView.cpp b/Engine/Plugins/Editor/AssetManagerEditor/Source/AssetManagerEditor/Private/Insights/Table/Widgets/STableTreeView.cpp
--- a/Engine/Plugins/Editor/AssetManagerEditor/Source/AssetManagerEditor/Private/Insights/Table/Widgets/STableTreeView.cpp
+++ b/Engine/Plugins/Editor/AssetManagerEditor/Source/AssetManagerEditor/Private/Insights/Table/Widgets/STableTreeView.cpp
@@ -1157,12 +1157,13 @@ BEGIN_SLATE_FUNCTION_BUILD_OPTIMIZATION
 TSharedRef<SWidget> STableTreeView::ConstructFilterConfiguratorButton()
 {
 	return SNew(SButton)
-		.ContentPadding(FMargin(-2.0f, 2.0f, -2.0f, 2.0f))
+		.ButtonStyle(FAppStyle::Get(), "SimpleButton")
 		.ToolTipText(LOCTEXT("FilterConfiguratorBtn_ToolTip", "Opens the filter configurator window."))
 		.OnClicked(this, &STableTreeView::FilterConfigurator_OnClicked)
 		[
 			SNew(SImage)
 			.Image(FInsightsStyle::GetBrush("Icons.ClassicFilterConfig"))
+			.ColorAndOpacity(FSlateColor::UseForeground())
 		];
 }
 END_SLATE_FUNCTION_BUILD_OPTIMIZATION
diff --git a/Engine/Plugins/Editor/AssetManagerEditor/Source/AssetManagerEditor/Private/SAssetAuditBrowser.cpp b/Engine/Plugins/Editor/AssetManagerEditor/Source/AssetManagerEditor/Private/SAssetAuditBrowser.cpp
--- a/Engine/Plugins/Editor/AssetManagerEditor/Source/AssetManagerEditor/Private/SAssetAuditBrowser.cpp
+++ b/Engine/Plugins/Editor/AssetManagerEditor/Source/AssetManagerEditor/Private/SAssetAuditBrowser.cpp
@@ -357,7 +357,8 @@ void SAssetAuditBrowser::Construct(const FArguments& InArgs)
 			]
 		];
 
-	TSharedRef<SWidget> DefaultContent =
+	this->ChildSlot
+	[
 		SNew(SVerticalBox)
 		+SVerticalBox::Slot()
 		.AutoHeight()
@@ -508,40 +509,8 @@ void SAssetAuditBrowser::Construct(const FArguments& InArgs)
 			[
 				ContentBrowserModule.Get().CreateAssetPicker(Config)
 			]
-		];
-
-	if (ManagerEditorModule.ShouldEnableTreeViewInAssetAudit())
-	{
-		TSharedRef<FAssetTable> AssetTable = MakeShared<FAssetTable>();
-		AssetTable->Reset();
-		AssetTable->SetDisplayName(FText::FromString(TEXT("AssetTable")));
-
-		this->ChildSlot
-		[
-			SNew(SSplitter)
-			.Orientation(EOrientation::Orient_Vertical)
-			+ SSplitter::Slot()
-			.Value(0.5f)
-			.MinSize(100.0f)
-			[
-				DefaultContent
-			]
-			+ SSplitter::Slot()
-			[
-				SAssignNew(AssetTableTreeView, SAssetTableTreeView, AssetTable)
-				.OnSelectionChanged(this, &SAssetAuditBrowser::TreeView_OnSelectionChanged)
-			]
-		];
-
-		AssetTableTreeView->RebuildTreeAsync(); // only needed for initial mock data
-	}
-	else
-	{
-		this->ChildSlot
-		[
-			DefaultContent
-		];
-	}
+		]
+	];
 
 	// Refresh our data sources to make sure we are displaying up to date info otherwise the user will need to
 	// hit "refresh" every time they open the dialog to be able to trust it
@@ -734,98 +703,6 @@ void SAssetAuditBrowser::GoToHistoryIndex(int32 InHistoryIdx)
 	}
 }
 
-void SAssetAuditBrowser::QuickPopulateAssetTableRow(FAssetTableRow& OutRow, const FAssetData& AssetData, FAssetTable& AssetTable) const
-{
-	OutRow.SoftObjectPath = AssetData.GetSoftObjectPath();
-
-	// Asset Type
-	FString AssetType;
-	if (UAssetDefinitionRegistry* AssetDefinitionRegistry = UAssetDefinitionRegistry::Get())
-	{
-		if (const UAssetDefinition* AssetDefinition = AssetDefinitionRegistry->GetAssetDefinitionForAsset(AssetData))
-		{
-			if (AssetDefinition)
-			{
-				AssetType = AssetDefinition->GetAssetDisplayName().ToString();
-			}
-		}
-	}
-	if (AssetType.IsEmpty())
-	{
-		if (!EditorModule->GetStringValueForCustomColumn(AssetData, FName(TEXT("Type")), AssetType))
-		{
-			AssetType = AssetData.AssetClassPath.ToString();
-		}
-	}
-	
-	OutRow.Type = AssetTable.StoreStr(AssetType);
-
-	// Asset Name
-	OutRow.Name = AssetTable.StoreStr(AssetData.AssetName.ToString());
-
-	// Path (without package or asset name)
-	FString ObjectPathString = AssetData.GetObjectPathString();
-	int Index;
-	ObjectPathString.FindLastChar(TEXT('/'), Index);
-	if (Index != INDEX_NONE)
-	{
-		ObjectPathString = ObjectPathString.Left(Index + 1);
-	}
-	OutRow.Path = AssetTable.StoreStr(ObjectPathString);
-}
-
-void SAssetAuditBrowser::PopulateAssetTableRow(FAssetTableRow& OutRow, const FAssetData& AssetData, FAssetTable& AssetTable) const
-{
-	QuickPopulateAssetTableRow(OutRow, AssetData, AssetTable);
-
-	FString Str;
-
-	if (EditorModule->GetStringValueForCustomColumn(AssetData, FPrimaryAssetId::PrimaryAssetTypeTag, Str))
-	{
-		OutRow.PrimaryType = AssetTable.StoreStr(Str);
-	}
-	if (EditorModule->GetStringValueForCustomColumn(AssetData, FPrimaryAssetId::PrimaryAssetNameTag, Str))
-	{
-		OutRow.PrimaryName = AssetTable.StoreStr(Str);
-	}
-
-	EditorModule->GetIntegerValueForCustomColumn(AssetData, IAssetManagerEditorModule::StageChunkCompressedSizeName, OutRow.StagedCompressedSize);
-	EditorModule->GetIntegerValueForCustomColumn(AssetData, IAssetManagerEditorModule::TotalUsageName, OutRow.TotalUsageCount);
-
-	if (EditorModule->GetStringValueForCustomColumn(AssetData, IAssetManagerEditorModule::ChunksName, Str))
-	{
-		OutRow.Chunks = AssetTable.StoreStr(Str);
-	}
-
-	if (EditorModule->GetStringValueForCustomColumn(AssetData, IAssetManagerEditorModule::PluginName, Str))
-	{
-		OutRow.PluginName = AssetTable.StoreStr(Str);
-	}
-
-	OutRow.NativeClass = AssetTable.StoreStr(AssetData.AssetClassPath.GetAssetName().ToString());
-
-	// Sets the color based on asset type.
-	const uint32 AssetTypeHash = GetTypeHash(FStringView(OutRow.Type));
-	OutRow.Color = USlateThemeManager::Get().GetColor((EStyleColor)((uint32)EStyleColor::AccentBlue + AssetTypeHash % 8));
-}
-
-void SAssetAuditBrowser::TreeView_OnSelectionChanged(const TArray<TSharedPtr<UE::Insights::FTableTreeNode>> InSelectedNodes)
-{
-	if (!AssetTableTreeView.IsValid())
-	{
-		return;
-	}
-
-	//TODO: Find and select the corresponding asset in the Asset Picker list view.
-	TSharedPtr<FAssetTreeNode> SelectedAssetNode = AssetTableTreeView->GetSingleSelectedAssetNode();
-	if (SelectedAssetNode.IsValid())
-	{
-		check(SelectedAssetNode->IsValidAsset());
-		const FAssetTableRow& AssetTableRow = SelectedAssetNode->GetAssetChecked();
-		//...
-	}
-}
-
 void SAssetAuditBrowser::RefreshAssetView()
 {
 	FARFilter Filter;
@@ -845,190 +722,6 @@ void SAssetAuditBrowser::RefreshAssetView()
 	}
 
 	SetFilterDelegate.Execute(Filter);
-
-	RefreshAssetTableTreeView(Filter);
-}
-
-// Refresh list of assets for the tree view
-void SAssetAuditBrowser::RefreshAssetTableTreeView(const FARFilter& Filter)
-{
-	if (!AssetTableTreeView.IsValid())
-	{
-		return;
-	}
-
-	TSharedPtr<FAssetTable> AssetTable = AssetTableTreeView->GetAssetTable();
-	if (!AssetTable.IsValid())
-	{
-		return;
-	}
-
-	UE_LOG(LogInsights, Log, TEXT("[AssetTree] Build asset table..."));
-
-	// TODO: We probably shouldn't do any of this unless the asset registry has finished loading. Otherwise bad things seem to happen.
-
-	UE::Insights::FStopwatch Stopwatch;
-	Stopwatch.Start();
-
-	// Clears all tree nodes (that references the previous assets).
-	AssetTable->SetVisibleAssetCount(0);
-	AssetTableTreeView->RebuildTree(true);
-
-	// Now is safe to clear the previous assets.
-	AssetTable->ClearAllData();
-
-	TMap<FAssetData, int32> AssetToIndexMap;
-
-	TArray<FAssetData> SourceAssets;
-	const FAssetManagerEditorRegistrySource* RegistrySource = EditorModule->GetCurrentRegistrySource();
-
-	if (RegistrySource->GetOwnedRegistryState())
-	{
-		FARCompiledFilter CompiledFilter;
-		AssetRegistry->CompileFilter(Filter, CompiledFilter);
-		RegistrySource->GetOwnedRegistryState()->GetAssets(CompiledFilter, TSet<FName>(), SourceAssets);
-	}
-	else
-	{
-		AssetRegistry->GetAssets(Filter, SourceAssets);
-	}
-
-	for (int32 SourceAssetIndex = 0; SourceAssetIndex < SourceAssets.Num(); SourceAssetIndex++)
-	{
-		TArray<FAssetData> AssetsInSourcePackage;
-		AssetRegistry->GetAssetsByPackageName(SourceAssets[SourceAssetIndex].PackageName, AssetsInSourcePackage, /*bIncludeOnlyOnDiskAssets*/ true); // Only use on disk assets to avoid creating FAssetData for everything in memory
-
-		for (FAssetData& SourceAsset : AssetsInSourcePackage)
-		{
-			if (AssetToIndexMap.Find(SourceAsset) == nullptr)
-			{
-				FAssetTableRow AssetRow;
-				PopulateAssetTableRow(AssetRow, SourceAsset, *AssetTable);
-				AssetTable->AddAsset(AssetRow);
-				AssetToIndexMap.Add(SourceAsset, AssetTable->GetTotalAssetCount() - 1);
-			}
-		}
-	}
-
-	int32 EntryIndex = AssetTable->GetAssets().Num();
-
-	TSet<FName> AlreadyProcessedPackages;
-	// Go over all SourceAssets. We will add more to the end as we find we need to add them.
-	for (int32 SourceAssetIndex = 0; SourceAssetIndex < SourceAssets.Num(); SourceAssetIndex++)
-	{
-		// There are potentially multiple assets in a package, each with an FAssetData
-		// We are only dealing in package dependencies, though, so don't process the same package twice
-		// But we will need to conceptually assign all the package dependencies to all assets in the package below
-		if (AlreadyProcessedPackages.Contains(SourceAssets[SourceAssetIndex].PackageName))
-		{
-			continue;
-		}
-		else
-		{
-			AlreadyProcessedPackages.Add(SourceAssets[SourceAssetIndex].PackageName);
-		}
-
-		FAssetIdentifier CurrentIdentifier(SourceAssets[SourceAssetIndex].PackageName);
-		// We'll have to add the dependencies to all these entries
-		TArray<FAssetData> AssetsInSourcePackage;
-		AssetRegistry->GetAssetsByPackageName(SourceAssets[SourceAssetIndex].PackageName, AssetsInSourcePackage, /*bIncludeOnlyOnDiskAssets*/ true); // Only use on disk assets to avoid creating FAssetData for everything in memory
-
-		// Get the dependencies
-		TArray<FAssetIdentifier> DependencyList;
-		RegistrySource->GetDependencies(CurrentIdentifier, DependencyList, UE::AssetRegistry::EDependencyCategory::Package);
-
-		TArray<int32> IndicesOfDependenciesInRowTableToAddToCurrentSourceAssetRow;
-
-		for (const FAssetIdentifier& Dependency : DependencyList)
-		{
-			auto GatherIndices = [&](FAssetData const* DependencyAsset)
-			{
-				if (int32* DependencyIndex = AssetToIndexMap.Find(*DependencyAsset))
-				{
-					// Do we already know about the asset? If so, great.
-					IndicesOfDependenciesInRowTableToAddToCurrentSourceAssetRow.Add(*DependencyIndex);
-				}
-				else
-				{
-					// We don't know about this asset yet. Create a new row for it and add it to the source list for further dependency analysis
-					FAssetTableRow NewRow;
-					PopulateAssetTableRow(NewRow, *DependencyAsset, *AssetTable);
-					AssetTable->AddAsset(NewRow);
-					IndicesOfDependenciesInRowTableToAddToCurrentSourceAssetRow.Add(AssetTable->GetTotalAssetCount() - 1);
-					AssetToIndexMap.Add(*DependencyAsset, AssetTable->GetTotalAssetCount() - 1);
-					SourceAssets.Add(*DependencyAsset);
-				}
-			};
-
-			if (RegistrySource->GetOwnedRegistryState())
-			{
-				TArrayView<FAssetData const* const> AssetsInDependencyPackage = RegistrySource->GetOwnedRegistryState()->GetAssetsByPackageName(Dependency.PackageName);
-				for (FAssetData const* const DependencyAsset : AssetsInDependencyPackage)
-				{
-					GatherIndices(DependencyAsset);
-				}
-			}
-			else
-			{
-				TArray<FAssetData> AssetsInDependencyPackage;
-				AssetRegistry->GetAssetsByPackageName(Dependency.PackageName, AssetsInDependencyPackage, /*only on disk assets*/true);
-				for (FAssetData& DependencyAsset : AssetsInDependencyPackage)
-				{
-					GatherIndices(&DependencyAsset);
-				}
-			}
-		}
-
-		if (IndicesOfDependenciesInRowTableToAddToCurrentSourceAssetRow.Num())
-		{
-			// We found some dependencies. Let's add them.
-			for (const FAssetData& SourceAssetData : AssetsInSourcePackage)
-			{
-				int32* RowIndex = AssetToIndexMap.Find(SourceAssetData);
-				if (RowIndex == nullptr)
-				{
-					UE_LOG(LogInsights, Warning, TEXT("Failed to find asset %s in package %s, source asset index %d. Asset registry loading was %s"), *SourceAssetData.AssetName.ToString(), *SourceAssetData.PackageName.ToString(), SourceAssetIndex, AssetRegistry->IsLoadingAssets() ? TEXT("INCOMPLETE") : TEXT("complete"));
-				}
-				if (ensure(RowIndex != nullptr))
-				{
-					for (int32 Index : IndicesOfDependenciesInRowTableToAddToCurrentSourceAssetRow)
-					{
-						AssetTable->GetAsset(*RowIndex)->Dependencies.AddUnique(Index);
-					}
-				}
-			}
-
-			for (int32 Index : IndicesOfDependenciesInRowTableToAddToCurrentSourceAssetRow)
-			{
-				// Now for each of those dependencies, add this source asset row as a referencer
-				if (FAssetTableRow* DependentRow = AssetTable->GetAsset(Index))
-				{
-					for (const FAssetData& SourceAssetData : AssetsInSourcePackage)
-					{
-						int32* RowIndex = AssetToIndexMap.Find(SourceAssetData);
-						if (ensure(RowIndex != nullptr))
-						{
-							DependentRow->Referencers.AddUnique(*RowIndex);
-						}
-					}
-				}
-			}
-		}
-		AssetTable->SetVisibleAssetCount(AssetTable->GetAssets().Num());
-	}
-
-	Stopwatch.Stop();
-	const double TotalTime = Stopwatch.GetAccumulatedTime();
-	UE_LOG(LogInsights, Log, TEXT("[AssetTree] Asset table rebuilt in %.4fs (%d visible + %d hidden = %d assets)"),
-		TotalTime, AssetTable->GetVisibleAssetCount(), AssetTable->GetHiddenAssetCount(), AssetTable->GetTotalAssetCount());
-
-	const FAssetTableStringStore& StringStore = AssetTable->GetStringStore();
-	UE_LOG(LogInsights, Log, TEXT("[AssetTree] String store: %d strings (%lld bytes) --> %d strings (%lld bytes | %lld bytes) %.0f%%"),
-		StringStore.GetNumInputStrings(), StringStore.GetTotalInputStringSize(),
-		StringStore.GetNumStrings(), StringStore.GetTotalStringSize(), StringStore.GetAllocatedSize(),
-		(double)StringStore.GetAllocatedSize() * 100.0 / (double)StringStore.GetTotalInputStringSize());
-
-	AssetTableTreeView->RebuildTreeAsync();
 }
 
 FReply SAssetAuditBrowser::ClearAssets()
@@ -1168,37 +861,6 @@ bool SAssetAuditBrowser::HandleFilterAsset(const FAssetData& InAssetData) const
 
 void SAssetAuditBrowser::OnAssetSelected(const FAssetData& InAssetData)
 {
-	// Find and select the corresponding asset in the asset tree view.
-	if (AssetTableTreeView.IsValid())
-	{
-		TSharedPtr<FAssetTable> AssetTable = AssetTableTreeView->GetAssetTable();
-		if (AssetTable.IsValid())
-		{
-			FAssetTableRow SelectedAsset;
-			QuickPopulateAssetTableRow(SelectedAsset, InAssetData, *AssetTable);
-
-			int32 RowIndex = -1;
-
-			const TArray<FAssetTableRow>& Assets = AssetTable->GetAssets();
-			int32 AssetCount = AssetTable->GetVisibleAssetCount();
-			for (int32 Index = 0; Index < AssetCount; ++Index)
-			{
-				const FAssetTableRow& Asset = Assets[Index];
-				if (Asset.GetType() == SelectedAsset.GetType() &&
-					Asset.GetPath() == SelectedAsset.GetPath() &&
-					Asset.GetName() == SelectedAsset.GetName())
-				{
-					RowIndex = Index;
-					break;
-				}
-			}
-
-			if (RowIndex >= 0)
-			{
-				AssetTableTreeView->SelectNodeByTableRowIndex(RowIndex);
-			}
-		}
-	}
 }
 
 TSharedRef<SWidget> SAssetAuditBrowser::GenerateSourceComboItem(TSharedPtr<FString> InItem)
@@ -1260,10 +922,6 @@ void SAssetAuditBrowser::SetCurrentRegistrySource(const FAssetManagerEditorRegis
 
 void SAssetAuditBrowser::OnClose()
 {
-	if (AssetTableTreeView.IsValid())
-	{
-		AssetTableTreeView->OnClose();
-	}
 }
 
 #undef LOCTEXT_NAMESPACE
diff --git a/Engine/Plugins/Editor/AssetManagerEditor/Source/AssetManagerEditor/Private/SAssetAuditBrowser.h b/Engine/Plugins/Editor/AssetManagerEditor/Source/AssetManagerEditor/Private/SAssetAuditBrowser.h
--- a/Engine/Plugins/Editor/AssetManagerEditor/Source/AssetManagerEditor/Private/SAssetAuditBrowser.h
+++ b/Engine/Plugins/Editor/AssetManagerEditor/Source/AssetManagerEditor/Private/SAssetAuditBrowser.h
@@ -12,14 +12,6 @@ class FUICommandList;
 class SMenuAnchor;
 class STextBlock;
 
-namespace UE
-{
-	namespace Insights
-	{
-		class FTableTreeNode;
-	}
-}
-
 //////////////////////////////////////////////////////////////////////////
 // SAssetAuditBrowser
 
@@ -140,16 +132,6 @@ protected:
 	/** Refresh the asset view with a new filter */
 	void RefreshAssetView();
 
-	void RefreshAssetTableTreeView(const FARFilter& Filter);
-
-	/** Populates an asset table row (view model for the asset tree view) with asset data (only identification data; type+path+name). */
-	void QuickPopulateAssetTableRow(class FAssetTableRow& OutRow, const FAssetData& AssetData, class FAssetTable& AssetTable) const;
-
-	/** Populates an asset table row (view model for the asset tree view) with asset data. */
-	void PopulateAssetTableRow(class FAssetTableRow& OutRow, const FAssetData& AssetData, class FAssetTable& AssetTable) const;
-
-	void TreeView_OnSelectionChanged(const TArray<TSharedPtr<UE::Insights::FTableTreeNode>> InSelectedNodes);
-
 	/** Commands handled by this widget */
 	TSharedPtr<FUICommandList> Commands;
 
@@ -189,7 +171,4 @@ protected:
 	class IAssetRegistry* AssetRegistry;
 	class UAssetManager* AssetManager;
 	class IAssetManagerEditorModule* EditorModule;
-
-	/** The Tree View widget */
-	TSharedPtr<class SAssetTableTreeView> AssetTableTreeView;
 };
diff --git a/Engine/Plugins/Editor/AssetManagerEditor/Source/AssetManagerEditor/Private/TreeView/AssetTable.h b/Engine/Plugins/Editor/AssetManagerEditor/Source/AssetManagerEditor/Private/TreeView/AssetTable.h
--- a/Engine/Plugins/Editor/AssetManagerEditor/Source/AssetManagerEditor/Private/TreeView/AssetTable.h
+++ b/Engine/Plugins/Editor/AssetManagerEditor/Source/AssetManagerEditor/Private/TreeView/AssetTable.h
@@ -3,10 +3,11 @@
 #pragma once
 
 #include "CoreMinimal.h"
-#include "Containers/StringView.h"
-#include "Misc/CString.h"
+
 #include "AssetRegistry/AssetData.h"
+#include "Containers/StringView.h"
 #include "Insights/Table/ViewModels/Table.h"
+#include "Misc/CString.h"
 
 namespace UE
 {
@@ -135,7 +136,7 @@ struct FAssetTableDependencySizes
 class FAssetTableRow
 {
 	friend class FAssetTable;
-	friend class SAssetAuditBrowser;
+	friend class SAssetTableTreeView;
 
 public:
 	FAssetTableRow()
diff --git a/Engine/Plugins/Editor/AssetManagerEditor/Source/AssetManagerEditor/Private/TreeView/SAssetTableTreeView.cpp b/Engine/Plugins/Editor/AssetManagerEditor/Source/AssetManagerEditor/Private/TreeView/SAssetTableTreeView.cpp
--- a/Engine/Plugins/Editor/AssetManagerEditor/Source/AssetManagerEditor/Private/TreeView/SAssetTableTreeView.cpp
+++ b/Engine/Plugins/Editor/AssetManagerEditor/Source/AssetManagerEditor/Private/TreeView/SAssetTableTreeView.cpp
@@ -2,27 +2,19 @@
 
 #include "SAssetTableTreeView.h"
 
+#include "AssetDefinitionRegistry.h"
+#include "AssetManagerEditorModule.h"
+#include "AssetRegistry/AssetRegistryModule.h"
 #include "Containers/Set.h"
 #include "ContentBrowserModule.h"
+#include "DesktopPlatformModule.h"
+#include "Editor.h"
+#include "Engine/AssetManager.h"
 #include "Framework/MultiBox/MultiBoxBuilder.h"
 #include "GenericPlatform/GenericPlatformMisc.h"
-#include "Modules/ModuleManager.h"
-#include "SlateOptMacros.h"
-#include "Styling/AppStyle.h"
-#include "Styling/StyleColors.h"
-#include "Widgets/Images/SImage.h"
-#include "Widgets/Input/SCheckBox.h"
-#include "Widgets/SToolTip.h"
-#include "DesktopPlatformModule.h"
-#include "HAL/PlatformFileManager.h"
 #include "HAL/FileManager.h"
-#include "Logging/MessageLog.h"
-#include "Editor.h"
-#include "Subsystems/AssetEditorSubsystem.h"
+#include "HAL/PlatformFileManager.h"
 #include "IContentBrowserSingleton.h"
-#include "AssetManagerEditorModule.h"
-#include "AssetRegistry/AssetRegistryModule.h"
-
 #include "Insights/Common/InsightsStyle.h"
 #include "Insights/Common/Log.h"
 #include "Insights/Common/Stopwatch.h"
@@ -30,9 +22,19 @@
 #include "Insights/Table/ViewModels/TableCellValueFormatter.h"
 #include "Insights/Table/ViewModels/TableColumn.h"
 #include "Insights/Table/ViewModels/TreeNodeGrouping.h"
+#include "Logging/MessageLog.h"
+#include "Modules/ModuleManager.h"
+#include "SlateOptMacros.h"
+#include "Styling/AppStyle.h"
+#include "Styling/StyleColors.h"
+#include "Subsystems/AssetEditorSubsystem.h"
 #include "TreeView/AssetTable.h"
 #include "TreeView/AssetTreeNode.h"
 #include "TreeView/AssetDependencyGrouping.h"
+#include "Widgets/Images/SImage.h"
+#include "Widgets/Input/SButton.h"
+#include "Widgets/Input/SCheckBox.h"
+#include "Widgets/SToolTip.h"
 
 #include <limits>
 #include <memory>
@@ -59,10 +61,17 @@
 {
 	this->OnSelectionChanged = InArgs._OnSelectionChanged;
 
+	FAssetRegistryModule& AssetRegistryModule = FModuleManager::LoadModuleChecked<FAssetRegistryModule>(TEXT("AssetRegistry"));
+	AssetRegistry = &AssetRegistryModule.Get();
+	AssetManager = &UAssetManager::Get();
+	EditorModule = &IAssetManagerEditorModule::Get();
+
 	STableTreeView::ConstructWidget(InTablePtr);
 
 	CreateGroupings();
 	CreateSortings();
+
+	RequestOpenRegistry();
 }
 
 ////////////////////////////////////////////////////////////////////////////////////////////////////
@@ -79,72 +88,24 @@
 {
 	UE::Insights::STableTreeView::Tick(AllottedGeometry, InCurrentTime, InDeltaTime);
 
-	if (bNeedsToRebuild && !bIsUpdateRunning)
+	if (!bIsUpdateRunning)
 	{
-		RebuildTree(true);
-		bNeedsToRebuild = false;
+		if (bNeedsToOpenRegistry)
+		{
+			bNeedsToOpenRegistry = false;
+			OpenRegistry();
+		}
+		if (bNeedsToRefreshAssets && !AssetRegistry->IsLoadingAssets())
+		{
+			bNeedsToRefreshAssets = false;
+			RefreshAssets();
+		}
+		if (bNeedsToRebuild)
+		{
+			bNeedsToRebuild = false;
+			RebuildTree(true);
+		}
 	}
 }
 
 ////////////////////////////////////////////////////////////////////////////////////////////////////
-
-void SAssetTableTreeView::RebuildTree(bool bResync)
-{
-	if (!bResync)
-	{
-		// There are no incremental updates.
-		return;
-	}
-
-	UE::Insights::FStopwatch Stopwatch;
-	Stopwatch.Start();
-
-	UE::Insights::FStopwatch SyncStopwatch;
-	SyncStopwatch.Start();
-
-	CancelCurrentAsyncOp();
-
-	const int32 PreviousNodeCount = TableRowNodes.Num();
-	TableRowNodes.Empty();
-
-	TSharedPtr<FAssetTable> AssetTable = GetAssetTable();
-	const int32 VisibleAssetCount = AssetTable.IsValid() ? AssetTable->GetVisibleAssetCount() : 0;
-
-	if (VisibleAssetCount > 0)
-	{
-		UE_LOG(LogInsights, Log, TEXT("[AssetTree] Creating %d asset nodes (previously: %d nodes)..."), VisibleAssetCount, PreviousNodeCount);
-
-		TableRowNodes.Reserve(VisibleAssetCount);
-
-		ensure(TableRowNodes.Num() == 0);
-		for (int32 AssetIndex = 0; AssetIndex < VisibleAssetCount; ++AssetIndex)
-		{
-			const FAssetTableRow* Asset = AssetTable->GetAsset(AssetIndex);
-
-			FName NodeName(Asset->GetName());
-			FAssetTreeNodePtr NodePtr = MakeShared<FAssetTreeNode>(NodeName, AssetTable, AssetIndex);
-			TableRowNodes.Add(NodePtr);
-		}
-		ensure(TableRowNodes.Num() == VisibleAssetCount);
-	}
-	else
-	{
-		UE_LOG(LogInsights, Log, TEXT("[AssetTree] Resetting tree (previously: %d nodes)..."), PreviousNodeCount);
-	}
-
-	SyncStopwatch.Stop();
-
-	UE_LOG(LogInsights, Log, TEXT("[AssetTree] Update tree..."));
-	UpdateTree();
-	TreeView->RebuildList();
-	TreeView->ClearSelection();
-	TreeView_OnSelectionChanged(nullptr, ESelectInfo::Type::Direct);
-
-	Stopwatch.Stop();
-	const double TotalTime = Stopwatch.GetAccumulatedTime();
-	const double SyncTime = SyncStopwatch.GetAccumulatedTime();
-	UE_LOG(LogInsights, Log, TEXT("[AssetTree] Tree view rebuilt in %.4fs (sync: %.4fs + update: %.4fs) --> %d asset nodes"),
-		TotalTime, SyncTime, TotalTime - SyncTime, TableRowNodes.Num());
-}
-
-////////////////////////////////////////////////////////////////////////////////////////////////////
@@ -171,39 +132,90 @@
 ////////////////////////////////////////////////////////////////////////////////////////////////////
 
 BEGIN_SLATE_FUNCTION_BUILD_OPTIMIZATION
-TSharedPtr<SWidget> SAssetTableTreeView::ConstructToolbar()
+void SAssetTableTreeView::ConstructHeaderArea(TSharedRef<SVerticalBox> InWidgetContent)
 {
-	TSharedPtr<SHorizontalBox> Box = SNew(SHorizontalBox);
+	InWidgetContent->AddSlot()
+		.VAlign(VAlign_Center)
+		.AutoHeight()
+		.Padding(2.0f)
+		[
+			SNew(SHorizontalBox)
 
-	Box->AddSlot()
-		.AutoWidth()
-		.VAlign(VAlign_Center)
-		.Padding(0.0f, 0.0f, 4.0f, 0.0f)
-		[
-			SNew(STextBlock)
-			.Text(LOCTEXT("Preset", "Preset:"))
+			+ SHorizontalBox::Slot()
+			.AutoWidth()
+			.Padding(0.0f, 0.0f, 4.0f, 0.0f)
+			.VAlign(VAlign_Center)
+			[
+				SNew(SButton)
+				.ButtonStyle(FAppStyle::Get(), "SimpleButton")
+				.ToolTipText(this, &SAssetTableTreeView::GetOpenRegistryToolTipText)
+				.IsEnabled(this, &SAssetTableTreeView::CanChangeRegistry)
+				.OnClicked(this, &SAssetTableTreeView::OnClickedOpenRegistry)
+				[
+					SNew(SImage)
+					.Image(FAppStyle::Get().GetBrush("Icons.FolderOpen"))
+					.ColorAndOpacity(FSlateColor::UseForeground())
+				]
+			]
+
+			+ SHorizontalBox::Slot()
+			.FillWidth(1.0f)
+			.VAlign(VAlign_Center)
+			[
+				ConstructSearchBox()
+			]
+
+			+ SHorizontalBox::Slot()
+			.AutoWidth()
+			.Padding(4.0f, 0.0f, 0.0f, 0.0f)
+			.VAlign(VAlign_Center)
+			[
+				ConstructFilterConfiguratorButton()
+			]
 		];
 
-	Box->AddSlot()
-		.AutoWidth()
+	InWidgetContent->AddSlot()
 		.VAlign(VAlign_Center)
+		.AutoHeight()
+		.Padding(2.0f)
 		[
-			SNew(SBox)
-			.MinDesiredWidth(150.0f)
+			SNew(SHorizontalBox)
+
+			+ SHorizontalBox::Slot()
+			.FillWidth(1.0f)
+			.VAlign(VAlign_Center)
 			[
-				SAssignNew(PresetComboBox, SComboBox<TSharedRef<UE::Insights::ITableTreeViewPreset>>)
-				.ToolTipText(this, &SAssetTableTreeView::ViewPreset_GetSelectedToolTipText)
-				.OptionsSource(GetAvailableViewPresets())
-				.OnSelectionChanged(this, &SAssetTableTreeView::ViewPreset_OnSelectionChanged)
-				.OnGenerateWidget(this, &SAssetTableTreeView::ViewPreset_OnGenerateWidget)
+				ConstructHierarchyBreadcrumbTrail()
+			]
+
+			+ SHorizontalBox::Slot()
+			.AutoWidth()
+			.VAlign(VAlign_Center)
+			.Padding(4.0f, 0.0f, 4.0f, 0.0f)
+			[
+				SNew(STextBlock)
+				.Text(LOCTEXT("Preset", "Preset:"))
+			]
+
+			+ SHorizontalBox::Slot()
+			.AutoWidth()
+			.VAlign(VAlign_Center)
+			[
+				SNew(SBox)
+				.MinDesiredWidth(150.0f)
 				[
-					SNew(STextBlock)
-					.Text(this, &SAssetTableTreeView::ViewPreset_GetSelectedText)
+					SAssignNew(PresetComboBox, SComboBox<TSharedRef<UE::Insights::ITableTreeViewPreset>>)
+					.ToolTipText(this, &SAssetTableTreeView::ViewPreset_GetSelectedToolTipText)
+					.OptionsSource(GetAvailableViewPresets())
+					.OnSelectionChanged(this, &SAssetTableTreeView::ViewPreset_OnSelectionChanged)
+					.OnGenerateWidget(this, &SAssetTableTreeView::ViewPreset_OnGenerateWidget)
+					[
+						SNew(STextBlock)
+						.Text(this, &SAssetTableTreeView::ViewPreset_GetSelectedText)
+					]
 				]
 			]
 		];
-
-	return Box;
 }
 END_SLATE_FUNCTION_BUILD_OPTIMIZATION
 
@@ -584,6 +596,8 @@
 	ReusableLineBuffer->Reset();
 }
 
+////////////////////////////////////////////////////////////////////////////////////////////////////
+
 void SAssetTableTreeView::ExportDependencyData() const
 {
 	IDesktopPlatform* DesktopPlatform = FDesktopPlatformModule::Get();
@@ -666,18 +680,11 @@
 	TArray<FAssetData> Assets;
 	for (int32 SelectionIndex : SelectedIndices)
 	{
-		const FAssetManagerEditorRegistrySource* RegistrySource = IAssetManagerEditorModule::Get().GetCurrentRegistrySource();
 		const FSoftObjectPath& SoftObjectPath = GetAssetTable()->GetAssetChecked(SelectionIndex).GetSoftObjectPath();
 		FAssetData AssetData;
-		if (RegistrySource->GetOwnedRegistryState())
+		if (IsRegistrySourceValid())
 		{
-			AssetData = *RegistrySource->GetOwnedRegistryState()->GetAssetByObjectPath(SoftObjectPath);
+			AssetData = *RegistrySource.GetOwnedRegistryState()->GetAssetByObjectPath(SoftObjectPath);
 		}
-		else
-		{
-			// Once we move this out of asset audit browser window and give it its own registry state, this code path can be killed
-			// it exists to handle the "Editor" case as a registry source
-			AssetData = FModuleManager::LoadModuleChecked<FAssetRegistryModule>(TEXT("AssetRegistry")).Get().GetAssetByObjectPath(SoftObjectPath, /*IncludeOnlyOnDiskAsset*/true, /*SkipARFilteredAssets*/false);
-		}
 		Assets.Add(AssetData);
 	}
@@ -696,7 +703,7 @@
 
 	FCanExecuteAction HasSelectionAndRegistrySourceAndCanExecute = FCanExecuteAction::CreateLambda([this]()
 		{
-			return (IAssetManagerEditorModule::Get().GetCurrentRegistrySource() != nullptr) && GetAssetTable() && (SelectedIndices.Num() > 0);
+			return IsRegistrySourceValid() && GetAssetTable() && (SelectedIndices.Num() > 0);
 		});
 
 	MenuBuilder.BeginSection("Asset", LOCTEXT("ContextMenu_Section_Asset", "Asset"));
@@ -719,7 +726,7 @@
 		MenuBuilder.AddMenuEntry
 			(
 				LOCTEXT("ContextMenu_EditAssetsLabel", "Edit..."),
-				LOCTEXT("ContextMenu_EditAssets", "Opens the selected aset in the relevant editor."),
+				LOCTEXT("ContextMenu_EditAssets", "Opens the selected asset in the relevant editor."),
 				FSlateIcon(),
 				EditSelectedAssets,
 				NAME_None,
@@ -739,7 +746,7 @@
 		MenuBuilder.AddMenuEntry
 		(
 			LOCTEXT("ContextMenu_FindInContentBrowserLabel", "Find in Content Browser..."),
-			LOCTEXT("ContextMenu_FindInContentBrowser", "Browses to the associated asset and selects it in the most recently used Content Browser (summoning one if necessary)"),
+			LOCTEXT("ContextMenu_FindInContentBrowser", "Browses to the associated asset and selects it in the most recently used Content Browser (summoning one if necessary)."),
 			FSlateIcon(),
 			FindInContentBrowser,
 			NAME_None,
@@ -764,7 +771,7 @@
 		MenuBuilder.AddMenuEntry
 		(
 			LOCTEXT("ContextMenu_OpenReferenceViewerLabel", "Reference Viewer..."),
-			LOCTEXT("ContextMenu_OpenReferenceViewer", "Launches the reference viewer showing the selected assets' references"),
+			LOCTEXT("ContextMenu_OpenReferenceViewer", "Launches the reference viewer showing the selected assets' references."),
 			FSlateIcon(),
 			OpenReferenceViewer,
 			NAME_None,
@@ -785,7 +792,7 @@
 		MenuBuilder.AddMenuEntry
 		(
 			LOCTEXT("ContextMenu_ExportDependenciesLabel", "Export Dependencies..."),
-			LOCTEXT("ContextMenu_ExportDependencies", "Export dependency CSVs for the selected asset"),
+			LOCTEXT("ContextMenu_ExportDependencies", "Exports dependency CSVs for the selected asset."),
 			FSlateIcon(),
 			ExportDependenciesAction,
 			NAME_None,
@@ -797,6 +804,58 @@
 
 ////////////////////////////////////////////////////////////////////////////////////////////////////
 
+void SAssetTableTreeView::RequestOpenRegistry()
+{
+	bNeedsToOpenRegistry = true;
+	FooterLeftText = LOCTEXT("FooterLeftTextFmt_OpenRegistry_Filtered", "Opening registry... please wait...");
+}
+
+////////////////////////////////////////////////////////////////////////////////////////////////////
+
+void SAssetTableTreeView::OpenRegistry()
+{
+	RegistrySource.SourceName = FAssetManagerEditorRegistrySource::CustomSourceName;
+	if (EditorModule->PopulateRegistrySource(&RegistrySource))
+	{
+		RequestRefreshAssets();
+	}
+}
+
+////////////////////////////////////////////////////////////////////////////////////////////////////
+
+FText SAssetTableTreeView::GetOpenRegistryToolTipText() const
+{
+	return FText::Format(LOCTEXT("OpenRegistryButtonToolTip", "Opens a new development asset registry.\nCurrent: {0}\nTimestamp: {1}"),
+		FText::FromString(RegistrySource.SourceFilename),
+		FText::FromString(RegistrySource.SourceTimestamp));
+}
+
+////////////////////////////////////////////////////////////////////////////////////////////////////
+
+bool SAssetTableTreeView::CanChangeRegistry() const
+{
+	return !bIsUpdateRunning && !bNeedsToOpenRegistry && !bNeedsToRefreshAssets && !bNeedsToRebuild;
+}
+
+////////////////////////////////////////////////////////////////////////////////////////////////////
+
+FReply SAssetTableTreeView::OnClickedOpenRegistry()
+{
+	CancelCurrentAsyncOp();
+	RequestOpenRegistry();
+	return FReply::Handled();
+}
+
+////////////////////////////////////////////////////////////////////////////////////////////////////
+
+void SAssetTableTreeView::RequestRefreshAssets()
+{
+	bNeedsToRefreshAssets = true;
+	FooterLeftText = LOCTEXT("FooterLeftTextFmt_RefreshAssets_Filtered", "Pre-analyzing disk sizes... please wait...");
+}
+
+////////////////////////////////////////////////////////////////////////////////////////////////////
+
 FText SAssetTableTreeView::GetFooterLeftText() const
 {
 	return FooterLeftText;
@@ -883,7 +942,7 @@
 	FAssetTreeNodePtr NewSelectedAssetNode;
 	int32 NewlySelectedAssetRowIndex = -1;
 	TSet<int32> SelectionSetIndices;
-	
+
 	for (const UE::Insights::FTableTreeNodePtr& Node : SelectedNodes)
 	{
 		if (Node->Is<FAssetTreeNode>() && Node->As<FAssetTreeNode>().IsValidAsset())
@@ -1020,4 +1079,303 @@
 
 ////////////////////////////////////////////////////////////////////////////////////////////////////
 
+void SAssetTableTreeView::QuickPopulateAssetTableRow(FAssetTableRow& OutRow, const FAssetData& AssetData, FAssetTable& AssetTable) const
+{
+	OutRow.SoftObjectPath = AssetData.GetSoftObjectPath();
+
+	// Asset Type
+	FString AssetType;
+	if (UAssetDefinitionRegistry* AssetDefinitionRegistry = UAssetDefinitionRegistry::Get())
+	{
+		const UAssetDefinition* AssetDefinition = AssetDefinitionRegistry->GetAssetDefinitionForAsset(AssetData);
+		if (AssetDefinition)
+		{
+			AssetType = AssetDefinition->GetAssetDisplayName().ToString();
+		}
+	}
+	if (AssetType.IsEmpty() && IsRegistrySourceValid())
+	{
+		if (!EditorModule->GetStringValueForCustomColumn(AssetData, FName(TEXT("Type")), AssetType, &RegistrySource))
+		{
+			AssetType = AssetData.AssetClassPath.ToString();
+		}
+	}
+	OutRow.Type = AssetTable.StoreStr(AssetType);
+
+	// Asset Name
+	OutRow.Name = AssetTable.StoreStr(AssetData.AssetName.ToString());
+
+	// Path (without package or asset name)
+	FString ObjectPathString = AssetData.GetObjectPathString();
+	int Index;
+	ObjectPathString.FindLastChar(TEXT('/'), Index);
+	if (Index != INDEX_NONE)
+	{
+		ObjectPathString = ObjectPathString.Left(Index + 1);
+	}
+	OutRow.Path = AssetTable.StoreStr(ObjectPathString);
+}
+
+////////////////////////////////////////////////////////////////////////////////////////////////////
+
+void SAssetTableTreeView::PopulateAssetTableRow(FAssetTableRow& OutRow, const FAssetData& AssetData, FAssetTable& AssetTable) const
+{
+	ensure(IsRegistrySourceValid());
+	QuickPopulateAssetTableRow(OutRow, AssetData, AssetTable);
+
+	FString Str;
+
+	if (EditorModule->GetStringValueForCustomColumn(AssetData, FPrimaryAssetId::PrimaryAssetTypeTag, Str, &RegistrySource))
+	{
+		OutRow.PrimaryType = AssetTable.StoreStr(Str);
+	}
+	if (EditorModule->GetStringValueForCustomColumn(AssetData, FPrimaryAssetId::PrimaryAssetNameTag, Str, &RegistrySource))
+	{
+		OutRow.PrimaryName = AssetTable.StoreStr(Str);
+	}
+
+	EditorModule->GetIntegerValueForCustomColumn(AssetData, IAssetManagerEditorModule::StageChunkCompressedSizeName, OutRow.StagedCompressedSize, &RegistrySource);
+	EditorModule->GetIntegerValueForCustomColumn(AssetData, IAssetManagerEditorModule::TotalUsageName, OutRow.TotalUsageCount, &RegistrySource);
+
+	if (EditorModule->GetStringValueForCustomColumn(AssetData, IAssetManagerEditorModule::ChunksName, Str, &RegistrySource))
+	{
+		OutRow.Chunks = AssetTable.StoreStr(Str);
+	}
+
+	if (EditorModule->GetStringValueForCustomColumn(AssetData, IAssetManagerEditorModule::PluginName, Str, &RegistrySource))
+	{
+		OutRow.PluginName = AssetTable.StoreStr(Str);
+	}
+
+	OutRow.NativeClass = AssetTable.StoreStr(AssetData.AssetClassPath.GetAssetName().ToString());
+
+	// Sets the color based on asset type.
+	const uint32 AssetTypeHash = GetTypeHash(FStringView(OutRow.Type));
+	OutRow.Color = USlateThemeManager::Get().GetColor((EStyleColor)((uint32)EStyleColor::AccentBlue + AssetTypeHash % 8));
+}
+
+////////////////////////////////////////////////////////////////////////////////////////////////////
+
+// Refresh list of assets for the tree view
+void SAssetTableTreeView::RefreshAssets()
+{
+	TSharedPtr<FAssetTable> AssetTable = GetAssetTable();
+	if (!AssetTable.IsValid())
+	{
+		return;
+	}
+
+	UE_LOG(LogInsights, Log, TEXT("[AssetTree] Build asset table..."));
+
+	UE::Insights::FStopwatch Stopwatch;
+	Stopwatch.Start();
+
+	// Clears all tree nodes (that references the previous assets).
+	AssetTable->SetVisibleAssetCount(0);
+	RebuildTree(true);
+
+	// Now is safe to clear the previous assets.
+	AssetTable->ClearAllData();
+
+	if (IsRegistrySourceValid())
+	{
+		TMap<FAssetData, int32> AssetToIndexMap;
+
+		TArray<FAssetData> SourceAssets;
+
+		RegistrySource.GetOwnedRegistryState()->GetAllAssets(TSet<FName>(), SourceAssets);
+
+		for (int32 SourceAssetIndex = 0; SourceAssetIndex < SourceAssets.Num(); SourceAssetIndex++)
+		{
+			TArray<FAssetData> AssetsInSourcePackage;
+			AssetRegistry->GetAssetsByPackageName(SourceAssets[SourceAssetIndex].PackageName, AssetsInSourcePackage, /*bIncludeOnlyOnDiskAssets*/ true); // Only use on disk assets to avoid creating FAssetData for everything in memory
+
+			for (FAssetData& SourceAsset : AssetsInSourcePackage)
+			{
+				if (AssetToIndexMap.Find(SourceAsset) == nullptr)
+				{
+					FAssetTableRow AssetRow;
+					PopulateAssetTableRow(AssetRow, SourceAsset, *AssetTable);
+					AssetTable->AddAsset(AssetRow);
+					AssetToIndexMap.Add(SourceAsset, AssetTable->GetTotalAssetCount() - 1);
+				}
+			}
+		}
+
+		int32 EntryIndex = AssetTable->GetAssets().Num();
+
+		TSet<FName> AlreadyProcessedPackages;
+		// Go over all SourceAssets. We will add more to the end as we find we need to add them.
+		for (int32 SourceAssetIndex = 0; SourceAssetIndex < SourceAssets.Num(); SourceAssetIndex++)
+		{
+			// There are potentially multiple assets in a package, each with an FAssetData
+			// We are only dealing in package dependencies, though, so don't process the same package twice
+			// But we will need to conceptually assign all the package dependencies to all assets in the package below
+			if (AlreadyProcessedPackages.Contains(SourceAssets[SourceAssetIndex].PackageName))
+			{
+				continue;
+			}
+			else
+			{
+				AlreadyProcessedPackages.Add(SourceAssets[SourceAssetIndex].PackageName);
+			}
+
+			FAssetIdentifier CurrentIdentifier(SourceAssets[SourceAssetIndex].PackageName);
+			// We'll have to add the dependencies to all these entries
+			TArray<FAssetData> AssetsInSourcePackage;
+			AssetRegistry->GetAssetsByPackageName(SourceAssets[SourceAssetIndex].PackageName, AssetsInSourcePackage, /*bIncludeOnlyOnDiskAssets*/ true); // Only use on disk assets to avoid creating FAssetData for everything in memory
+
+			// Get the dependencies
+			TArray<FAssetIdentifier> DependencyList;
+			RegistrySource.GetDependencies(CurrentIdentifier, DependencyList, UE::AssetRegistry::EDependencyCategory::Package);
+
+			TArray<int32> IndicesOfDependenciesInRowTableToAddToCurrentSourceAssetRow;
+
+			for (const FAssetIdentifier& Dependency : DependencyList)
+			{
+				TArrayView<FAssetData const* const> AssetsInDependencyPackage = RegistrySource.GetOwnedRegistryState()->GetAssetsByPackageName(Dependency.PackageName);
+				for (FAssetData const* const DependencyAsset : AssetsInDependencyPackage)
+				{
+					if (int32* DependencyIndex = AssetToIndexMap.Find(*DependencyAsset))
+					{
+						// Do we already know about the asset? If so, great.
+						IndicesOfDependenciesInRowTableToAddToCurrentSourceAssetRow.Add(*DependencyIndex);
+					}
+					else
+					{
+						// We don't know about this asset yet. Create a new row for it and add it to the source list for further dependency analysis
+						FAssetTableRow NewRow;
+						PopulateAssetTableRow(NewRow, *DependencyAsset, *AssetTable);
+						AssetTable->AddAsset(NewRow);
+						IndicesOfDependenciesInRowTableToAddToCurrentSourceAssetRow.Add(AssetTable->GetTotalAssetCount() - 1);
+						AssetToIndexMap.Add(*DependencyAsset, AssetTable->GetTotalAssetCount() - 1);
+						SourceAssets.Add(*DependencyAsset);
+					}
+				}
+			}
+
+			if (IndicesOfDependenciesInRowTableToAddToCurrentSourceAssetRow.Num())
+			{
+				// We found some dependencies. Let's add them.
+				for (const FAssetData& SourceAssetData : AssetsInSourcePackage)
+				{
+					int32* RowIndex = AssetToIndexMap.Find(SourceAssetData);
+					if (RowIndex == nullptr)
+					{
+						UE_LOG(LogInsights, Warning, TEXT("Failed to find asset %s in package %s, source asset index %d. Asset registry loading was %s"), *SourceAssetData.AssetName.ToString(), *SourceAssetData.PackageName.ToString(), SourceAssetIndex, AssetRegistry->IsLoadingAssets() ? TEXT("INCOMPLETE") : TEXT("complete"));
+					}
+					if (ensure(RowIndex != nullptr))
+					{
+						for (int32 Index : IndicesOfDependenciesInRowTableToAddToCurrentSourceAssetRow)
+						{
+							AssetTable->GetAsset(*RowIndex)->Dependencies.AddUnique(Index);
+						}
+					}
+				}
+
+				for (int32 Index : IndicesOfDependenciesInRowTableToAddToCurrentSourceAssetRow)
+				{
+					// Now for each of those dependencies, add this source asset row as a referencer
+					if (FAssetTableRow* DependentRow = AssetTable->GetAsset(Index))
+					{
+						for (const FAssetData& SourceAssetData : AssetsInSourcePackage)
+						{
+							int32* RowIndex = AssetToIndexMap.Find(SourceAssetData);
+							if (ensure(RowIndex != nullptr))
+							{
+								DependentRow->Referencers.AddUnique(*RowIndex);
+							}
+						}
+					}
+				}
+			}
+			AssetTable->SetVisibleAssetCount(AssetTable->GetAssets().Num());
+		}
+	}
+
+	Stopwatch.Stop();
+	const double TotalTime = Stopwatch.GetAccumulatedTime();
+	UE_LOG(LogInsights, Log, TEXT("[AssetTree] Asset table rebuilt in %.4fs (%d visible + %d hidden = %d assets)"),
+		TotalTime, AssetTable->GetVisibleAssetCount(), AssetTable->GetHiddenAssetCount(), AssetTable->GetTotalAssetCount());
+
+	const FAssetTableStringStore& StringStore = AssetTable->GetStringStore();
+	UE_LOG(LogInsights, Log, TEXT("[AssetTree] String store: %d strings (%lld bytes) --> %d strings (%lld bytes | %lld bytes) %.0f%%"),
+		StringStore.GetNumInputStrings(), StringStore.GetTotalInputStringSize(),
+		StringStore.GetNumStrings(), StringStore.GetTotalStringSize(), StringStore.GetAllocatedSize(),
+		(double)StringStore.GetAllocatedSize() * 100.0 / (double)StringStore.GetTotalInputStringSize());
+
+	RequestRebuildTree();
+}
+
+////////////////////////////////////////////////////////////////////////////////////////////////////
+
+void SAssetTableTreeView::RequestRebuildTree()
+{
+	bNeedsToRebuild = true;
+	FooterLeftText = LOCTEXT("FooterLeftTextFmt_RebuildTree_Filtered", "Rebuilding tree... please wait...");
+}
+
+////////////////////////////////////////////////////////////////////////////////////////////////////
+
+void SAssetTableTreeView::RebuildTree(bool bResync)
+{
+	if (!bResync)
+	{
+		// There are no incremental updates.
+		return;
+	}
+
+	UE::Insights::FStopwatch Stopwatch;
+	Stopwatch.Start();
+
+	UE::Insights::FStopwatch SyncStopwatch;
+	SyncStopwatch.Start();
+
+	CancelCurrentAsyncOp();
+
+	const int32 PreviousNodeCount = TableRowNodes.Num();
+	TableRowNodes.Empty();
+
+	TSharedPtr<FAssetTable> AssetTable = GetAssetTable();
+	const int32 VisibleAssetCount = AssetTable.IsValid() ? AssetTable->GetVisibleAssetCount() : 0;
+
+	if (VisibleAssetCount > 0)
+	{
+		UE_LOG(LogInsights, Log, TEXT("[AssetTree] Creating %d asset nodes (previously: %d nodes)..."), VisibleAssetCount, PreviousNodeCount);
+
+		TableRowNodes.Reserve(VisibleAssetCount);
+
+		ensure(TableRowNodes.Num() == 0);
+		for (int32 AssetIndex = 0; AssetIndex < VisibleAssetCount; ++AssetIndex)
+		{
+			const FAssetTableRow* Asset = AssetTable->GetAsset(AssetIndex);
+
+			FName NodeName(Asset->GetName());
+			FAssetTreeNodePtr NodePtr = MakeShared<FAssetTreeNode>(NodeName, AssetTable, AssetIndex);
+			TableRowNodes.Add(NodePtr);
+		}
+		ensure(TableRowNodes.Num() == VisibleAssetCount);
+	}
+	else
+	{
+		UE_LOG(LogInsights, Log, TEXT("[AssetTree] Resetting tree (previously: %d nodes)..."), PreviousNodeCount);
+	}
+
+	SyncStopwatch.Stop();
+
+	UE_LOG(LogInsights, Log, TEXT("[AssetTree] Update tree..."));
+	UpdateTree();
+	TreeView->RebuildList();
+	TreeView->ClearSelection();
+	TreeView_OnSelectionChanged(nullptr, ESelectInfo::Type::Direct);
+
+	Stopwatch.Stop();
+	const double TotalTime = Stopwatch.GetAccumulatedTime();
+	const double SyncTime = SyncStopwatch.GetAccumulatedTime();
+	UE_LOG(LogInsights, Log, TEXT("[AssetTree] Tree view rebuilt in %.4fs (sync: %.4fs + update: %.4fs) --> %d asset nodes"),
+		TotalTime, SyncTime, TotalTime - SyncTime, TableRowNodes.Num());
+}
+
+////////////////////////////////////////////////////////////////////////////////////////////////////
+
 #undef LOCTEXT_NAMESPACE
diff --git a/Engine/Plugins/Editor/AssetManagerEditor/Source/AssetManagerEditor/Private/TreeView/SAssetTableTreeView.h b/Engine/Plugins/Editor/AssetManagerEditor/Source/AssetManagerEditor/Private/TreeView/SAssetTableTreeView.h
--- a/Engine/Plugins/Editor/AssetManagerEditor/Source/AssetManagerEditor/Private/TreeView/SAssetTableTreeView.h
+++ b/Engine/Plugins/Editor/AssetManagerEditor/Source/AssetManagerEditor/Private/TreeView/SAssetTableTreeView.h
@@ -5,6 +5,7 @@
 #include "CoreMinimal.h"
 #include "Delegates/DelegateCombinations.h"
 #include "Widgets/DeclarativeSyntaxSupport.h"
+#include "AssetManagerEditorModule.h"
 
 #include "TreeView/AssetTable.h"
 #include "Insights/Table/Widgets/STableTreeView.h"
@@ -37,7 +38,7 @@ public:
 	 */
 	void Construct(const FArguments& InArgs, TSharedPtr<FAssetTable> InTablePtr);
 
-	virtual TSharedPtr<SWidget> ConstructToolbar() override;
+	virtual void ConstructHeaderArea(TSharedRef<SVerticalBox> InWidgetContent);
 	virtual TSharedPtr<SWidget> ConstructFooter() override;
 
 	TSharedPtr<FAssetTable> GetAssetTable() { return StaticCastSharedPtr<FAssetTable>(GetTable()); }
@@ -59,7 +60,6 @@ public:
 	 * @param bResync - If true, it forces a resync even if the list did not changed since last sync.
 	 */
 	virtual void RebuildTree(bool bResync);
-	virtual void RebuildTreeAsync() { bNeedsToRebuild = true; }
 
 	////////////////////////////////////////////////////////////////////////////////////////////////////
 	// IAsyncOperationStatusProvider implementation
@@ -72,11 +72,17 @@ public:
 
 	TSharedPtr<FAssetTreeNode> GetSingleSelectedAssetNode() const { return SelectedAssetNode; }
 
-protected:
+private:
 	virtual void InternalCreateGroupings() override;
 
 	virtual void ExtendMenu(FMenuBuilder& MenuBuilder) override;
 
+	void RequestOpenRegistry();
+	void OpenRegistry();
+	FText GetOpenRegistryToolTipText() const;
+	bool CanChangeRegistry() const;
+	FReply OnClickedOpenRegistry();
+
 	FText GetFooterLeftText() const;
 	FText GetFooterCenterText1() const;
 	FText GetFooterCenterText2() const;
@@ -84,7 +90,6 @@ protected:
 
 	virtual void TreeView_OnSelectionChanged(UE::Insights::FTableTreeNodePtr SelectedItem, ESelectInfo::Type SelectInfo) override;
 
-private:
 	void InitAvailableViewPresets();
 
 	void ExportDependencyData() const;
@@ -93,7 +98,22 @@ private:
 
 	void CalculateBaseAndMarginalCostForSelection(TSet<int32>& SelectionSetIndices, int64* OutTotalSizeMultiplyUsed, int64* OutTotalSizeSingleUse) const;
 
+	/** Populates an asset table row (view model for the asset tree view) with asset data (only identification data; type+path+name). */
+	void QuickPopulateAssetTableRow(class FAssetTableRow& OutRow, const FAssetData& AssetData, class FAssetTable& AssetTable) const;
+
+	/** Populates an asset table row (view model for the asset tree view) with asset data. */
+	void PopulateAssetTableRow(class FAssetTableRow& OutRow, const FAssetData& AssetData, class FAssetTable& AssetTable) const;
+
+	bool IsRegistrySourceValid() const { return RegistrySource.GetOwnedRegistryState() != nullptr; }
+
+	void RequestRefreshAssets();
+	void RefreshAssets();
+
+	void RequestRebuildTree();
+
 private:
+	bool bNeedsToOpenRegistry = false;
+	bool bNeedsToRefreshAssets = false;
 	bool bNeedsToRebuild = false;
 
 	FText FooterLeftText;
@@ -105,4 +125,20 @@ private:
 	FOnSelectionChanged OnSelectionChanged;
 	TSharedPtr<FAssetTreeNode> SelectedAssetNode;
 	TSet<int32> SelectedIndices;
+
+	/** List of valid registry sources */
+	TArray<TSharedPtr<FString>> SourceComboList;
+
+	/**
+	* Text block where we put the timestamp for the asset registry so the user knows
+	* if they are looking at super out of date data.
+	*/
+	TSharedPtr<STextBlock> RegistrySourceTimeText;
+
+	FAssetManagerEditorRegistrySource RegistrySource;
+
+	/** Cached interfaces */
+	class IAssetRegistry* AssetRegistry;
+	class UAssetManager* AssetManager;
+	class IAssetManagerEditorModule* EditorModule;
 };
diff --git a/Engine/Plugins/Editor/AssetManagerEditor/Source/AssetManagerEditor/Public/AssetManagerEditorModule.h b/Engine/Plugins/Editor/AssetManagerEditor/Source/AssetManagerEditor/Public/AssetManagerEditorModule.h
--- a/Engine/Plugins/Editor/AssetManagerEditor/Source/AssetManagerEditor/Public/AssetManagerEditorModule.h
+++ b/Engine/Plugins/Editor/AssetManagerEditor/Source/AssetManagerEditor/Public/AssetManagerEditorModule.h
@@ -386,13 +386,13 @@ public:
 	static const FName PluginName;
 
 	/** Gets the value of a "virtual" column for an asset data, this will query the AssetManager for you and takes current platform into account. Returns true and sets out parameter if found */
-	virtual bool GetStringValueForCustomColumn(const FAssetData& AssetData, FName ColumnName, FString& OutValue) = 0;
+	virtual bool GetStringValueForCustomColumn(const FAssetData& AssetData, FName ColumnName, FString& OutValue, const FAssetManagerEditorRegistrySource* OverrideRegistrySource = nullptr) = 0;
 
 	/** Gets a display text value for a virtual column. Returns true and sets out parameter if found */
-	virtual bool GetDisplayTextForCustomColumn(const FAssetData& AssetData, FName ColumnName, FText& OutValue) = 0;
+	virtual bool GetDisplayTextForCustomColumn(const FAssetData& AssetData, FName ColumnName, FText& OutValue, const FAssetManagerEditorRegistrySource* OverrideRegistrySource = nullptr) = 0;
 
 	/** Gets the value of a "virtual" column for an asset data, this will query the AssetManager for you and takes current platform into account. Returns true and sets out parameter if found */
-	virtual bool GetIntegerValueForCustomColumn(const FAssetData& AssetData, FName ColumnName, int64& OutValue) = 0;
+	virtual bool GetIntegerValueForCustomColumn(const FAssetData& AssetData, FName ColumnName, int64& OutValue, const FAssetManagerEditorRegistrySource* OverrideRegistrySource = nullptr) = 0;
 
 	/** Returns the set of asset packages managed the given asset data, returns true if any found. This handles the fake chunk/primary asset types above */
 	virtual bool GetManagedPackageListForAssetData(const FAssetData& AssetData, TSet<FName>& ManagedPackageSet) = 0;
@@ -406,6 +406,9 @@ public:
 	/** Sets the current registry source, this loads the asset registry state if needed and may spawn a file load dialog for custom */
 	virtual void SetCurrentRegistrySource(const FString& SourceName) = 0;
 
+	/** Sets up a registry source. If InOutRegistrySource->SourceName is CustomSourceName, a dialog will be presented to select the asset registry to load */
+	virtual bool PopulateRegistrySource(FAssetManagerEditorRegistrySource* InOutRegistrySource) = 0;
+
 	/** Refreshes the management dictionary and all sources */
 	virtual void RefreshRegistryData() = 0;
 
@@ -423,9 +426,6 @@ public:
 	virtual void OpenAssetAuditUI(TArray<FAssetIdentifier> SelectedIdentifiers) = 0;
 	virtual void OpenAssetAuditUI(TArray<FName> SelectedPackages) = 0;
 
-	/** Returns true if the tree view will be enabled in the Asset Audit UI, the next time the Asset Audit window is open. */
-	virtual bool ShouldEnableTreeViewInAssetAudit() const = 0;
-
 	/** Spawns reference viewer, showing selected packages or identifiers */
 	virtual void OpenReferenceViewerUI(const TArray<FAssetIdentifier> SelectedIdentifiers, const FReferenceViewerParams ReferenceViewerParams = FReferenceViewerParams()) = 0;
 	virtual void OpenReferenceViewerUI(const TArray<FName> SelectedPackages, const FReferenceViewerParams ReferenceViewerParams = FReferenceViewerParams()) = 0;
