[Disk Size Tool] Refactored the value getters for FAssetTable so that string values get propagated to nodes that wrap a single asset for dependency grouping purposes. This also simplifies the string value getters generally.

diff --git a/Engine/Plugins/Editor/AssetManagerEditor/Source/AssetManagerEditor/Private/TreeView/AssetTable.cpp b/Engine/Plugins/Editor/AssetManagerEditor/Source/AssetManagerEditor/Private/TreeView/AssetTable.cpp
--- a/Engine/Plugins/Editor/AssetManagerEditor/Source/AssetManagerEditor/Private/TreeView/AssetTable.cpp
+++ b/Engine/Plugins/Editor/AssetManagerEditor/Source/AssetManagerEditor/Private/TreeView/AssetTable.cpp
@@ -36,6 +36,40 @@ const FName FAssetTableColumns::ChunksColumnId(TEXT("Chunks"));
 const FName FAssetTableColumns::NativeClassColumnId(TEXT("NativeClass"));
 const FName FAssetTableColumns::PluginNameColumnId(TEXT("PluginName"));
 
+////////////////////////////////////////////////////////////////////////////////////////////////////
+// FAssetTableStringValueGetterWithDependencyAggregationHandling
+////////////////////////////////////////////////////////////////////////////////////////////////////
+
+// A special value getter that handles the common case where a dependency grouping node is acting as a
+// proxy for a single asset and its dependencies.
+class FAssetTableStringValueGetterWithDependencyAggregationHandling : public UE::Insights::FTableCellValueGetter
+{
+public:
+
+	virtual const TCHAR* GetStringValueFromAssetTableRow(const FAssetTableRow& Asset) const = 0;
+
+	virtual const TOptional<UE::Insights::FTableCellValue> GetValue(const UE::Insights::FTableColumn& Column, const UE::Insights::FBaseTreeNode& Node) const override
+	{
+		using namespace UE::Insights;
+		const FTableTreeNode& NodePtr = static_cast<const FTableTreeNode&>(Node);
+		if (Node.IsGroup() && !NodePtr.GetRowId().HasValidIndex())
+		{
+			if (NodePtr.HasAggregatedValue(Column.GetId()))
+			{
+				return NodePtr.GetAggregatedValue(Column.GetId());
+			}
+		}
+		else //if (Node->Is<FAssetTreeNode>())
+		{
+			TSharedPtr<FAssetTable> AssetTable = StaticCastSharedPtr<FAssetTable>(NodePtr.GetParentTable().Pin());
+			const FAssetTableRow& Asset = AssetTable->GetAssetChecked(NodePtr.GetRowId().RowIndex);
+			return FTableCellValue(GetStringValueFromAssetTableRow(Asset));
+		}
+
+		return TOptional<FTableCellValue>();
+	}
+};
+
 ////////////////////////////////////////////////////////////////////////////////////////////////////
 // FAssetTableStringStore
 ////////////////////////////////////////////////////////////////////////////////////////////////////
@@ -551,29 +585,16 @@ void FAssetTable::AddDefaultColumns()
 
 		Column.SetDataType(ETableCellDataType::CString);
 
-		class FAssetTypeValueGetter : public FTableCellValueGetter
+		class FAssetTypeValueGetter : public FAssetTableStringValueGetterWithDependencyAggregationHandling 
 		{
 		public:
-			virtual const TOptional<FTableCellValue> GetValue(const FTableColumn& Column, const FBaseTreeNode& Node) const
-			{
-				if (Node.IsGroup())
-				{
-					const FTableTreeNode& NodePtr = static_cast<const FTableTreeNode&>(Node);
-					if (NodePtr.HasAggregatedValue(Column.GetId()))
-					{
-						return NodePtr.GetAggregatedValue(Column.GetId());
-					}
-				}
-				else //if (Node->Is<FAssetTreeNode>())
-				{
-					const FAssetTreeNode& TreeNode = static_cast<const FAssetTreeNode&>(Node);
-					const FAssetTableRow& Asset = TreeNode.GetAssetChecked();
-					return FTableCellValue(Asset.GetType());
-				}
 
-				return TOptional<FTableCellValue>();
+			virtual const TCHAR* GetStringValueFromAssetTableRow(const FAssetTableRow& Asset) const override
+			{
+				return Asset.GetType();
 			}
 		};
+
 		TSharedRef<ITableCellValueGetter> Getter = MakeShared<FAssetTypeValueGetter>();
 		Column.SetValueGetter(Getter);
 
@@ -606,29 +627,16 @@ void FAssetTable::AddDefaultColumns()
 
 		Column.SetDataType(ETableCellDataType::CString);
 
-		class FAssetNameValueGetter : public FTableCellValueGetter
+		class FAssetNameValueGetter : public FAssetTableStringValueGetterWithDependencyAggregationHandling
 		{
 		public:
-			virtual const TOptional<FTableCellValue> GetValue(const FTableColumn& Column, const FBaseTreeNode& Node) const
-			{
-				if (Node.IsGroup())
-				{
-					const FTableTreeNode& NodePtr = static_cast<const FTableTreeNode&>(Node);
-					if (NodePtr.HasAggregatedValue(Column.GetId()))
-					{
-						return NodePtr.GetAggregatedValue(Column.GetId());
-					}
-				}
-				else //if (Node->Is<FAssetTreeNode>())
-				{
-					const FAssetTreeNode& TreeNode = static_cast<const FAssetTreeNode&>(Node);
-					const FAssetTableRow& Asset = TreeNode.GetAssetChecked();
-					return FTableCellValue(Asset.GetName());
-				}
 
-				return TOptional<FTableCellValue>();
+			virtual const TCHAR* GetStringValueFromAssetTableRow(const FAssetTableRow& Asset) const override
+			{
+				return Asset.GetName();
 			}
 		};
+
 		TSharedRef<ITableCellValueGetter> Getter = MakeShared<FAssetNameValueGetter>();
 		Column.SetValueGetter(Getter);
 
@@ -661,29 +669,16 @@ void FAssetTable::AddDefaultColumns()
 
 		Column.SetDataType(ETableCellDataType::CString);
 
-		class FAssetPathValueGetter : public FTableCellValueGetter
+		class FAssetPathValueGetter : public FAssetTableStringValueGetterWithDependencyAggregationHandling
 		{
 		public:
-			virtual const TOptional<FTableCellValue> GetValue(const FTableColumn& Column, const FBaseTreeNode& Node) const
-			{
-				if (Node.IsGroup())
-				{
-					const FTableTreeNode& NodePtr = static_cast<const FTableTreeNode&>(Node);
-					if (NodePtr.HasAggregatedValue(Column.GetId()))
-					{
-						return NodePtr.GetAggregatedValue(Column.GetId());
-					}
-				}
-				else //if (Node->Is<FAssetTreeNode>())
-				{
-					const FAssetTreeNode& TreeNode = static_cast<const FAssetTreeNode&>(Node);
-					const FAssetTableRow& Asset = TreeNode.GetAssetChecked();
-					return FTableCellValue(Asset.GetPath());
-				}
 
-				return TOptional<FTableCellValue>();
+			virtual const TCHAR* GetStringValueFromAssetTableRow(const FAssetTableRow& Asset) const override
+			{
+				return Asset.GetPath();
 			}
 		};
+
 		TSharedRef<ITableCellValueGetter> Getter = MakeShared<FAssetPathValueGetter>();
 		Column.SetValueGetter(Getter);
 
@@ -716,29 +711,16 @@ void FAssetTable::AddDefaultColumns()
 
 		Column.SetDataType(ETableCellDataType::CString);
 
-		class FAssetPrimaryTypeValueGetter : public FTableCellValueGetter
+		class FAssetPrimaryTypeValueGetter : public FAssetTableStringValueGetterWithDependencyAggregationHandling
 		{
 		public:
-			virtual const TOptional<FTableCellValue> GetValue(const FTableColumn& Column, const FBaseTreeNode& Node) const
-			{
-				if (Node.IsGroup())
-				{
-					const FTableTreeNode& NodePtr = static_cast<const FTableTreeNode&>(Node);
-					if (NodePtr.HasAggregatedValue(Column.GetId()))
-					{
-						return NodePtr.GetAggregatedValue(Column.GetId());
-					}
-				}
-				else //if (Node->Is<FAssetTreeNode>())
-				{
-					const FAssetTreeNode& TreeNode = static_cast<const FAssetTreeNode&>(Node);
-					const FAssetTableRow& Asset = TreeNode.GetAssetChecked();
-					return FTableCellValue(Asset.GetPrimaryType());
-				}
 
-				return TOptional<FTableCellValue>();
+			virtual const TCHAR* GetStringValueFromAssetTableRow(const FAssetTableRow& Asset) const override
+			{
+				return Asset.GetPrimaryType();
 			}
 		};
+
 		TSharedRef<ITableCellValueGetter> Getter = MakeShared<FAssetPrimaryTypeValueGetter>();
 		Column.SetValueGetter(Getter);
 
@@ -771,29 +753,16 @@ void FAssetTable::AddDefaultColumns()
 
 		Column.SetDataType(ETableCellDataType::CString);
 
-		class FAssetPrimaryNameValueGetter : public FTableCellValueGetter
+		class FAssetPrimaryNameValueGetter : public FAssetTableStringValueGetterWithDependencyAggregationHandling
 		{
 		public:
-			virtual const TOptional<FTableCellValue> GetValue(const FTableColumn& Column, const FBaseTreeNode& Node) const
-			{
-				if (Node.IsGroup())
-				{
-					const FTableTreeNode& NodePtr = static_cast<const FTableTreeNode&>(Node);
-					if (NodePtr.HasAggregatedValue(Column.GetId()))
-					{
-						return NodePtr.GetAggregatedValue(Column.GetId());
-					}
-				}
-				else //if (Node->Is<FAssetTreeNode>())
-				{
-					const FAssetTreeNode& TreeNode = static_cast<const FAssetTreeNode&>(Node);
-					const FAssetTableRow& Asset = TreeNode.GetAssetChecked();
-					return FTableCellValue(Asset.GetPrimaryName());
-				}
 
-				return TOptional<FTableCellValue>();
+			virtual const TCHAR* GetStringValueFromAssetTableRow(const FAssetTableRow& Asset) const override
+			{
+				return Asset.GetPrimaryName();
 			}
 		};
+
 		TSharedRef<ITableCellValueGetter> Getter = MakeShared<FAssetPrimaryNameValueGetter>();
 		Column.SetValueGetter(Getter);
 
@@ -1143,29 +1112,16 @@ void FAssetTable::AddDefaultColumns()
 
 		Column.SetDataType(ETableCellDataType::CString);
 
-		class FChunksValueGetter : public FTableCellValueGetter
+		class FChunksValueGetter : public FAssetTableStringValueGetterWithDependencyAggregationHandling
 		{
 		public:
-			virtual const TOptional<FTableCellValue> GetValue(const FTableColumn& Column, const FBaseTreeNode& Node) const
-			{
-				if (Node.IsGroup())
-				{
-					const FTableTreeNode& NodePtr = static_cast<const FTableTreeNode&>(Node);
-					if (NodePtr.HasAggregatedValue(Column.GetId()))
-					{
-						return NodePtr.GetAggregatedValue(Column.GetId());
-					}
-				}
-				else //if (Node->Is<FAssetTreeNode>())
-				{
-					const FAssetTreeNode& TreeNode = static_cast<const FAssetTreeNode&>(Node);
-					const FAssetTableRow& Asset = TreeNode.GetAssetChecked();
-					return FTableCellValue(Asset.GetChunks());
-				}
 
-				return TOptional<FTableCellValue>();
+			virtual const TCHAR* GetStringValueFromAssetTableRow(const FAssetTableRow& Asset) const override
+			{
+				return Asset.GetChunks();
 			}
 		};
+
 		TSharedRef<ITableCellValueGetter> Getter = MakeShared<FChunksValueGetter>();
 		Column.SetValueGetter(Getter);
 
@@ -1198,29 +1154,16 @@ void FAssetTable::AddDefaultColumns()
 
 		Column.SetDataType(ETableCellDataType::CString);
 
-		class FNativeClassValueGetter : public FTableCellValueGetter
+		class FNativeClassValueGetter : public FAssetTableStringValueGetterWithDependencyAggregationHandling
 		{
 		public:
-			virtual const TOptional<FTableCellValue> GetValue(const FTableColumn& Column, const FBaseTreeNode& Node) const
-			{
-				if (Node.IsGroup())
-				{
-					const FTableTreeNode& NodePtr = static_cast<const FTableTreeNode&>(Node);
-					if (NodePtr.HasAggregatedValue(Column.GetId()))
-					{
-						return NodePtr.GetAggregatedValue(Column.GetId());
-					}
-				}
-				else //if (Node->Is<FAssetTreeNode>())
-				{
-					const FAssetTreeNode& TreeNode = static_cast<const FAssetTreeNode&>(Node);
-					const FAssetTableRow& Asset = TreeNode.GetAssetChecked();
-					return FTableCellValue(Asset.GetNativeClass());
-				}
 
-				return TOptional<FTableCellValue>();
+			virtual const TCHAR* GetStringValueFromAssetTableRow(const FAssetTableRow& Asset) const override
+			{
+				return Asset.GetNativeClass();
 			}
 		};
+
 		TSharedRef<ITableCellValueGetter> Getter = MakeShared<FNativeClassValueGetter>();
 		Column.SetValueGetter(Getter);
 
@@ -1253,30 +1196,17 @@ void FAssetTable::AddDefaultColumns()
 
 		Column.SetDataType(ETableCellDataType::CString);
 
-		class FNativeClassValueGetter : public FTableCellValueGetter
+		class FPluginNameValueGetter : public FAssetTableStringValueGetterWithDependencyAggregationHandling
 		{
 		public:
-			virtual const TOptional<FTableCellValue> GetValue(const FTableColumn& Column, const FBaseTreeNode& Node) const
-			{
-				if (Node.IsGroup())
-				{
-					const FTableTreeNode& NodePtr = static_cast<const FTableTreeNode&>(Node);
-					if (NodePtr.HasAggregatedValue(Column.GetId()))
-					{
-						return NodePtr.GetAggregatedValue(Column.GetId());
-					}
-				}
-				else //if (Node->Is<FAssetTreeNode>())
-				{
-					const FAssetTreeNode& TreeNode = static_cast<const FAssetTreeNode&>(Node);
-					const FAssetTableRow& Asset = TreeNode.GetAssetChecked();
-					return FTableCellValue(Asset.GetPluginName());
-				}
 
-				return TOptional<FTableCellValue>();
+			virtual const TCHAR* GetStringValueFromAssetTableRow(const FAssetTableRow& Asset) const override
+			{
+				return Asset.GetPluginName();
 			}
 		};
-		TSharedRef<ITableCellValueGetter> Getter = MakeShared<FNativeClassValueGetter>();
+
+		TSharedRef<ITableCellValueGetter> Getter = MakeShared<FPluginNameValueGetter>();
 		Column.SetValueGetter(Getter);
 
 		TSharedRef<ITableCellValueFormatter> Formatter = MakeShared<FCStringValueFormatterAsText>();
