AssetManagerEditor / Asset Disk Size: Added "Plugin Dependency Analysis" grouping and preset.

diff --git a/Engine/Plugins/Editor/AssetManagerEditor/Source/AssetManagerEditor/AssetManagerEditor.Build.cs b/Engine/Plugins/Editor/AssetManagerEditor/Source/AssetManagerEditor/AssetManagerEditor.Build.cs
--- a/Engine/Plugins/Editor/AssetManagerEditor/Source/AssetManagerEditor/AssetManagerEditor.Build.cs
+++ b/Engine/Plugins/Editor/AssetManagerEditor/Source/AssetManagerEditor/AssetManagerEditor.Build.cs
@@ -50,6 +50,7 @@ public class AssetManagerEditor : ModuleRules
 				"SourceControl",
 				"SourceControlWindows",
 				"UncontrolledChangelists",
+				"Projects"
 			}
 		);
 	}
diff --git a/Engine/Plugins/Editor/AssetManagerEditor/Source/AssetManagerEditor/Private/TreeView/AssetDependencyGrouping.cpp b/Engine/Plugins/Editor/AssetManagerEditor/Source/AssetManagerEditor/Private/TreeView/AssetDependencyGrouping.cpp
--- a/Engine/Plugins/Editor/AssetManagerEditor/Source/AssetManagerEditor/Private/TreeView/AssetDependencyGrouping.cpp
+++ b/Engine/Plugins/Editor/AssetManagerEditor/Source/AssetManagerEditor/Private/TreeView/AssetDependencyGrouping.cpp
@@ -8,6 +8,7 @@
 #define LOCTEXT_NAMESPACE "FAssetDependencyGrouping"
 
 INSIGHTS_IMPLEMENT_RTTI(FAssetDependencyGrouping)
+INSIGHTS_IMPLEMENT_RTTI(FPluginDependencyGrouping)
 
 ////////////////////////////////////////////////////////////////////////////////////////////////////
 // FAssetDependencyGrouping
@@ -54,34 +55,39 @@ void FAssetDependencyGrouping::GroupNodes(const TArray<UE::Insights::FTableTreeN
 		}
 
 		// For each (visible) asset, we will create the following hierarchy:
-		//
-		// By default, the dependencies node is collapsed.
+
+		// If the asset does not have dependencies, the asset node is added directly (no group node is created).
+		// |
+		// +-- [asset:{AssetName}]
+
+		// By default, the "Dependencies" node is collapsed.
 		// |
 		// +-- [group:{AssetName}] (self + dependencies)
-		// |   |
-		// |   +-- [group:_Self]
-		// |   |   |
-		// |   |   +-- [asset:{AssetName}]
-		// |   |
-		// |   +-- [group:Dependencies] (double click to expand)
+		//     |
+		//     +-- [group:_Self]
+		//     |   |
+		//     |   +-- [asset:{AssetName}]
+		//     |
+		//     +-- [group:Dependencies] (double click to expand)
 
-		// When dependencies node is expanded, it will be populated with actual dependencies.
+		// When the "Dependencies" node is expanded, it will be populated with actual dependencies.
 		// |
 		// +-- [group:{AssetName}] (self + dependencies)
-		// |   |
-		// |   +-- [group:_Self]
-		// |   |   |
-		// |   |   +-- [asset:{AssetName}]
-		// |   |
-		// |   +-- [group:Dependencies]
-		// |       |
-		// |       +-- [asset:{DependentAsset1}]
-		// |       |
-		// |       +-- [asset:{DependentAsset2}]
-		// |       |
-		// |       +-- [group:{DependentAsset3}] (self + dependencies)
-		// |       |
-		// |       ...
+		//     |
+		//     +-- [group:_Self]
+		//     |   |
+		//     |   +-- [asset:{AssetName}]
+		//     |
+		//     +-- [group:Dependencies]
+		//         |
+		//         +-- [asset:{DependentAsset1}]
+		//         |
+		//         +-- [asset:{DependentAsset2}]
+		//         |
+		//         +-- [group:{DependentAsset3}] (self + dependencies)
+		//         |   |
+		//         |   ...
+		//         ...
 
 		FAssetTreeNode& AssetNode = NodePtr->As<FAssetTreeNode>();
 		const FAssetTableRow& Asset = AssetNode.GetAssetChecked();
@@ -115,6 +121,148 @@ void FAssetDependencyGrouping::GroupNodes(const TArray<UE::Insights::FTableTreeN
 	}
 }
 
+////////////////////////////////////////////////////////////////////////////////////////////////////
+// FPluginDependencyGrouping
+////////////////////////////////////////////////////////////////////////////////////////////////////
+
+FPluginDependencyGrouping::FPluginDependencyGrouping()
+	: FTreeNodeGrouping(
+		LOCTEXT("Grouping_ByPluginDependency_ShortName", "Plugin Dependency"),
+		LOCTEXT("Grouping_ByPluginDependency_TitleName", "By Plugin Dependency"),
+		LOCTEXT("Grouping_ByPluginDependency_Desc", "Group assets based on their plugin. Each plugin group node will all show the dependent plugins."),
+		TEXT("Icons.Group.TreeItem"),
+		nullptr)
+{
+}
+
+////////////////////////////////////////////////////////////////////////////////////////////////////
+
+FPluginDependencyGrouping::~FPluginDependencyGrouping()
+{
+}
+
+////////////////////////////////////////////////////////////////////////////////////////////////////
+
+void FPluginDependencyGrouping::GroupNodes(const TArray<UE::Insights::FTableTreeNodePtr>& Nodes, UE::Insights::FTableTreeNode& ParentGroup, TWeakPtr<UE::Insights::FTable> InParentTable, UE::Insights::IAsyncOperationProgress& InAsyncOperationProgress) const
+{
+	using namespace UE::Insights;
+
+	ParentGroup.ClearChildren();
+
+	TSharedPtr<FAssetTable> AssetTable = StaticCastSharedPtr<FAssetTable>(InParentTable.Pin());
+	check(AssetTable.IsValid());
+
+	TMap<int32, FPluginSimpleGroupNode*> PluginIndexToGroupNodeMap;
+
+	for (FTableTreeNodePtr NodePtr : Nodes)
+	{
+		if (InAsyncOperationProgress.ShouldCancelAsyncOp())
+		{
+			return;
+		}
+
+		if (NodePtr->IsGroup() || !NodePtr->Is<FAssetTreeNode>())
+		{
+			ParentGroup.AddChildAndSetParent(NodePtr);
+			continue;
+		}
+
+		// For each (visible) asset, we will create the following hierarchy:
+
+		// If the plugin does not have dependencies, a simple group will be created.
+		// |
+		// +-- [group:{PluginName}] (self) // FPluginSimpleGroupNode
+		// |   |
+		// |   +-- [asset:{AssetName1}]
+		// |   |
+		// |   +-- [asset:{AssetName2}]
+		// |   |
+		// |   ...
+
+		// By default, the "Plugin Dependencies" node is collapsed.
+		// |
+		// +-- [group:{PluginName}] (self + dependencies) // FPluginAndDependenciesGroupNode
+		// |   |
+		// |   +-- [group:Plugin Dependencies] (double click to expand) // FPluginDependenciesGroupNode, lazy
+		// |   |
+		// |   +-- [group:{PluginName}] (self) // FPluginSimpleGroupNode
+		// |       |
+		// |       +-- [asset:{AssetName1}]
+		// |       |
+		// |       +-- [asset:{AssetName2}]
+		// |       |
+		// |       ...
+
+		// When the "Plugin Dependencies" node is expanded, it will be populated with actual plugin dependencies.
+		// |
+		// +-- [group:{PluginName}] (self + dependencies) // FPluginAndDependenciesGroupNode
+		// |   |
+		// |   +-- [group:Plugin Dependencies] // FPluginDependenciesGroupNode expanded
+		// |   |   |
+		// |   |   +-- [group:{DependentPlugin1}] (self + dependencies) // FPluginAndDependenciesGroupNode
+		// |   |   |   |
+		// |   |   |   +-- [group:Plugin Dependencies] (double click to expand) // FPluginDependenciesGroupNode, lazy
+		// |   |   |   |
+		// |   |   |   +-- [group:{DependentPlugin1}] (self) // FPluginSimpleGroupNode
+		// |   |   |       |
+		// |   |   |       +-- [asset:{Asset1a}]
+		// |   |   |       |
+		// |   |   |       +-- [asset:{Asset1b}]
+		// |   |   |       ...
+		// |   |   |
+		// |   |   +-- [group:{DependentPlugin2}] (self, no further dependencies) // FPluginSimpleGroupNode
+		// |   |   |   |
+		// |   |   |   +-- [asset:{Asset2a}]
+		// |   |   |   |
+		// |   |   |   +-- [asset:{Asset2b}]
+		// |   |   |   ...
+		// |   |   |
+		// |   |   ...
+		// |   |
+		// |   +-- [group:{PluginName}] (self) // FPluginSimpleGroupNode
+		// |       |
+		// |       +-- [asset:{Asset1a}]
+		// |       |
+		// |       +-- [asset:{Asset1b}]
+		// |       |
+		// |       ...
+
+		FAssetTreeNode& AssetNode = NodePtr->As<FAssetTreeNode>();
+		const FAssetTableRow& Asset = AssetNode.GetAssetChecked();
+
+		const TCHAR* PluginName = Asset.GetPluginName();
+		const FAssetTablePluginInfo& PluginInfo = AssetTable->GetOrCreatePluginInfo(PluginName);
+		int32 PluginIndex = AssetTable->GetIndexForPlugin(PluginName);
+
+		FPluginSimpleGroupNode* PluginGroup = PluginIndexToGroupNodeMap.FindRef(PluginIndex);
+		if (PluginGroup)
+		{
+			// Plugin group node (simple or with dependencies) was already created. Just add the current asset to it.
+			PluginGroup->AddChildAndSetParent(NodePtr);
+		}
+		else if (PluginInfo.PluginDependencies.Num() > 0)
+		{
+			// Create the Plugin Self+Dependencies group node and add the current asset to the Self group.
+			FName PluginAndDependenciesGroupName = AssetTable->GetNameForPlugin(PluginIndex);
+			TSharedPtr<FPluginAndDependenciesGroupNode> PluginAndDependenciesGroup = MakeShared<FPluginAndDependenciesGroupNode>(PluginAndDependenciesGroupName, AssetTable, PluginIndex);
+			ParentGroup.AddChildAndSetParent(PluginAndDependenciesGroup);
+			PluginGroup = PluginAndDependenciesGroup->CreateChildren().Get();
+			PluginGroup->AddChildAndSetParent(NodePtr);
+			PluginIndexToGroupNodeMap.Add(PluginIndex, PluginGroup);
+		}
+		else
+		{
+			// Create a simple Plugin group node and add the current asset to it.
+			FName PluginGroupName = AssetTable->GetNameForPlugin(PluginIndex);
+			TSharedPtr<FPluginSimpleGroupNode> PluginGroupPtr = MakeShared<FPluginSimpleGroupNode>(PluginGroupName, AssetTable, PluginIndex);
+			ParentGroup.AddChildAndSetParent(PluginGroupPtr);
+			PluginGroup = PluginGroupPtr.Get();
+			PluginGroup->AddChildAndSetParent(NodePtr);
+			PluginIndexToGroupNodeMap.Add(PluginIndex, PluginGroup);
+		}
+	}
+}
+
 ////////////////////////////////////////////////////////////////////////////////////////////////////
 ////////////////////////////////////////////////////////////////////////////////////////////////////
 
diff --git a/Engine/Plugins/Editor/AssetManagerEditor/Source/AssetManagerEditor/Private/TreeView/AssetDependencyGrouping.h b/Engine/Plugins/Editor/AssetManagerEditor/Source/AssetManagerEditor/Private/TreeView/AssetDependencyGrouping.h
--- a/Engine/Plugins/Editor/AssetManagerEditor/Source/AssetManagerEditor/Private/TreeView/AssetDependencyGrouping.h
+++ b/Engine/Plugins/Editor/AssetManagerEditor/Source/AssetManagerEditor/Private/TreeView/AssetDependencyGrouping.h
@@ -16,3 +16,14 @@ public:
 
 	virtual void GroupNodes(const TArray<UE::Insights::FTableTreeNodePtr>& Nodes, UE::Insights::FTableTreeNode& ParentGroup, TWeakPtr<UE::Insights::FTable> InParentTable, UE::Insights::IAsyncOperationProgress& InAsyncOperationProgress) const override;
 };
+
+class FPluginDependencyGrouping : public UE::Insights::FTreeNodeGrouping
+{
+	INSIGHTS_DECLARE_RTTI(FPluginDependencyGrouping, UE::Insights::FTreeNodeGrouping)
+
+public:
+	FPluginDependencyGrouping();
+	virtual ~FPluginDependencyGrouping();
+
+	virtual void GroupNodes(const TArray<UE::Insights::FTableTreeNodePtr>& Nodes, UE::Insights::FTableTreeNode& ParentGroup, TWeakPtr<UE::Insights::FTable> InParentTable, UE::Insights::IAsyncOperationProgress& InAsyncOperationProgress) const override;
+};
diff --git a/Engine/Plugins/Editor/AssetManagerEditor/Source/AssetManagerEditor/Private/TreeView/AssetTable.cpp b/Engine/Plugins/Editor/AssetManagerEditor/Source/AssetManagerEditor/Private/TreeView/AssetTable.cpp
--- a/Engine/Plugins/Editor/AssetManagerEditor/Source/AssetManagerEditor/Private/TreeView/AssetTable.cpp
+++ b/Engine/Plugins/Editor/AssetManagerEditor/Source/AssetManagerEditor/Private/TreeView/AssetTable.cpp
@@ -490,6 +490,61 @@ void FAssetTable::Reset()
 
 ////////////////////////////////////////////////////////////////////////////////////////////////////
 
+void FAssetTable::EnumeratePluginDependencies(const FAssetTablePluginInfo& InPlugin, TFunction<void(int32 PluginIndex)> Callback) const
+{
+	for (int32 DependencyIndex : InPlugin.PluginDependencies)
+	{
+		Callback(DependencyIndex);
+	}
+}
+
+////////////////////////////////////////////////////////////////////////////////////////////////////
+
+void FAssetTable::EnumerateAssetsForPlugin(const FAssetTablePluginInfo& InPlugin, TFunction<void(int32 AssetIndex)> Callback) const
+{
+	const int32 AssetCount = Assets.Num();
+	for (int32 AssetIndex = 0; AssetIndex < AssetCount; ++AssetIndex)
+	{
+		// The StringStore makes this comparison safe
+		if (Assets[AssetIndex].GetPluginName() == InPlugin.PluginName)
+		{
+			Callback(AssetIndex);
+		}
+	}
+}
+
+////////////////////////////////////////////////////////////////////////////////////////////////////
+
+FAssetTablePluginInfo& FAssetTable::GetOrCreatePluginInfo(const TCHAR* StoredPluginName)
+{
+	if (int32* PluginIndex = PluginNameToIndexMap.Find(StoredPluginName))
+	{
+		return Plugins[*PluginIndex];
+	}
+	else
+	{
+		FAssetTablePluginInfo NewPlugin;
+		NewPlugin.PluginName = StoredPluginName;
+		int32 NewIndex = Plugins.Num();
+		Plugins.Add(NewPlugin);
+		PluginNameToIndexMap.Add(StoredPluginName, NewIndex);
+		return Plugins[NewIndex];
+	}
+}
+
+////////////////////////////////////////////////////////////////////////////////////////////////////
+
+int32 FAssetTable::GetIndexForPlugin(const TCHAR* StoredPluginName) const
+{
+	if (const int32* Index = PluginNameToIndexMap.Find(StoredPluginName))
+	{
+		return *Index;
+	}
+	return -1;
+}
+
+////////////////////////////////////////////////////////////////////////////////////////////////////
+
 void FAssetTable::AddDefaultColumns()
 {
 	using namespace UE::Insights;
diff --git a/Engine/Plugins/Editor/AssetManagerEditor/Source/AssetManagerEditor/Private/TreeView/AssetTable.h b/Engine/Plugins/Editor/AssetManagerEditor/Source/AssetManagerEditor/Private/TreeView/AssetTable.h
--- a/Engine/Plugins/Editor/AssetManagerEditor/Source/AssetManagerEditor/Private/TreeView/AssetTable.h
+++ b/Engine/Plugins/Editor/AssetManagerEditor/Source/AssetManagerEditor/Private/TreeView/AssetTable.h
@@ -133,6 +133,13 @@ struct FAssetTableDependencySizes
 	int64 SharedDependenciesSize = 0;
 };
 
+struct FAssetTablePluginInfo
+{
+	TSet<int32> PluginDependencies;
+	const TCHAR* PluginName = nullptr;
+	int64 Size = -1;
+};
+
 class FAssetTableRow
 {
 	friend class FAssetTable;
@@ -226,6 +233,27 @@ public:
 
 	void AddAsset(const FAssetTableRow& AssetRow) { Assets.Add(AssetRow); }
 
+	FAssetTablePluginInfo& GetOrCreatePluginInfo(const TCHAR* StoredPluginName);
+
+	const FAssetTablePluginInfo& GetPluginInfoByIndex(int32 PluginIndex) const { return Plugins[PluginIndex]; }
+	const FAssetTablePluginInfo& GetPluginInfoByIndexChecked(int32 PluginIndex) const { check(IsValidPluginIndex(PluginIndex)); return Plugins[PluginIndex]; }
+	int32 GetIndexForPlugin(const TCHAR* StoredPluginName) const;
+	int32 GetNumPlugins() const { return Plugins.Num(); }
+
+	FName GetNameForPlugin(int32 InPluginIndex) const
+	{
+		if (IsValidPluginIndex(InPluginIndex))
+		{
+			return FName(Plugins[InPluginIndex].PluginName);
+		}
+		return FName("<InvalidPlugin>");
+	}
+
+	bool IsValidPluginIndex(int32 Index) const { return Index >= 0 && Index < Plugins.Num(); }
+
+	void EnumeratePluginDependencies(const FAssetTablePluginInfo& InPlugin, TFunction<void(int32 PluginIndex)> Callback) const;
+	void EnumerateAssetsForPlugin(const FAssetTablePluginInfo& InPlugin, TFunction<void(int32 AssetIndex)> Callback) const;
+
 	const FStringView StoreString(const FStringView InStr) { return StringStore.Store(InStr); }
 	const TCHAR* StoreStr(const FString& InStr) { return StringStore.Store(InStr).GetData(); }
 	const FAssetTableStringStore& GetStringStore() const { return StringStore; }
@@ -236,6 +264,8 @@ public:
 		Assets.Empty();
 		VisibleAssetCount = 0;
 		StringStore.Reset();
+		Plugins.Empty();
+		PluginNameToIndexMap.Empty();
 	}
 
 private:
@@ -243,6 +273,8 @@ private:
 
 private:
 	TArray<FAssetTableRow> Assets;
+	TArray<FAssetTablePluginInfo> Plugins;
+	TMap<const TCHAR*, int32, FDefaultSetAllocator, TStringPointerMapKeyFuncs_DEPRECATED<const TCHAR*, int32>> PluginNameToIndexMap;
 	int32 VisibleAssetCount = 0;
 	FAssetTableStringStore StringStore;
 };
diff --git a/Engine/Plugins/Editor/AssetManagerEditor/Source/AssetManagerEditor/Private/TreeView/AssetTreeNode.cpp b/Engine/Plugins/Editor/AssetManagerEditor/Source/AssetManagerEditor/Private/TreeView/AssetTreeNode.cpp
--- a/Engine/Plugins/Editor/AssetManagerEditor/Source/AssetManagerEditor/Private/TreeView/AssetTreeNode.cpp
+++ b/Engine/Plugins/Editor/AssetManagerEditor/Source/AssetManagerEditor/Private/TreeView/AssetTreeNode.cpp
@@ -12,6 +12,9 @@
 
 INSIGHTS_IMPLEMENT_RTTI(FAssetTreeNode)
 INSIGHTS_IMPLEMENT_RTTI(FAssetDependenciesGroupTreeNode)
+INSIGHTS_IMPLEMENT_RTTI(FPluginSimpleGroupNode)
+INSIGHTS_IMPLEMENT_RTTI(FPluginAndDependenciesGroupNode)
+INSIGHTS_IMPLEMENT_RTTI(FPluginDependenciesGroupNode)
 
 ////////////////////////////////////////////////////////////////////////////////////////////////////
 // FAssetTreeNode
@@ -99,6 +102,163 @@ bool FAssetDependenciesGroupTreeNode::OnLazyCreateChildren(TSharedPtr<UE::Insigh
 	return true;
 }
 
+////////////////////////////////////////////////////////////////////////////////////////////////////
+// FPluginSimpleGroupNode
+////////////////////////////////////////////////////////////////////////////////////////////////////
+
+const FSlateBrush* FPluginSimpleGroupNode::GetIcon() const
+{
+	return UE::Insights::FBaseTreeNode::GetDefaultIcon(true);
+}
+
+////////////////////////////////////////////////////////////////////////////////////////////////////
+
+FLinearColor FPluginSimpleGroupNode::GetColor() const
+{
+	return FLinearColor(1.0f, 1.0f, 1.0f, 1.0f);
+}
+
+////////////////////////////////////////////////////////////////////////////////////////////////////
+// FPluginSimpleGroupNode
+////////////////////////////////////////////////////////////////////////////////////////////////////
+
+void FPluginSimpleGroupNode::AddAssetChildrenNodes()
+{
+	// [this]
+	// |
+	// +-- [asset:{Asset1}]
+	// |
+	// +-- [asset:{Asset2}]
+	// ...
+
+	FAssetTable& AssetTable = GetAssetTableChecked();
+	if (AssetTable.IsValidPluginIndex(PluginIndex))
+	{
+		const FAssetTablePluginInfo& PluginInfo = AssetTable.GetPluginInfoByIndex(PluginIndex);
+		AssetTable.EnumerateAssetsForPlugin(PluginInfo, [this, AssetTablePtr = &AssetTable](int32 AssetIndex)
+			{
+				FName AssetNodeName(AssetTablePtr->GetAssetChecked(AssetIndex).GetName());
+				FAssetTreeNodePtr AssetNode = MakeShared<FAssetTreeNode>(AssetNodeName, GetAssetTableWeak(), AssetIndex);
+				AddChildAndSetParent(AssetNode);
+			});
+	}
+}
+
+////////////////////////////////////////////////////////////////////////////////////////////////////
+// FPluginAndDependenciesGroupNode
+////////////////////////////////////////////////////////////////////////////////////////////////////
+
+TSharedPtr<FPluginSimpleGroupNode> FPluginAndDependenciesGroupNode::CreateChildren()
+{
+	// [this]
+	// |
+	// +-- [group:Plugin Dependencies] (double click to expand) // FPluginDependenciesGroupNode, lazy
+	// |
+	// +-- [group:{PluginName}] (self) // FPluginSimpleGroupNode
+
+	FAssetTable& AssetTable = GetAssetTableChecked();
+	if (AssetTable.IsValidPluginIndex(PluginIndex))
+	{
+		// Create the Plugin Dependencies group node.
+		// The children nodes (list of dependent plugins) will be lazy created.
+		static FName DependenciesGroupName(TEXT("Plugin Dependencies"));
+		TSharedPtr<FPluginDependenciesGroupNode> DependenciesGroup = MakeShared<FPluginDependenciesGroupNode>(DependenciesGroupName, GetAssetTableWeak(), PluginIndex);
+		AddChildAndSetParent(DependenciesGroup);
+
+		// Create the Plugin Self group node (where asset nodes will be added).
+		FName PluginGroupName = AssetTable.GetNameForPlugin(PluginIndex);
+		TSharedPtr<FPluginSimpleGroupNode> PluginGroup = MakeShared<FPluginSimpleGroupNode>(PluginGroupName, GetAssetTableWeak(), PluginIndex);
+		AddChildAndSetParent(PluginGroup);
+		return PluginGroup;
+	}
+	return SharedThis(this);
+}
+
+////////////////////////////////////////////////////////////////////////////////////////////////////
+// FPluginDependenciesGroupNode
+////////////////////////////////////////////////////////////////////////////////////////////////////
+
+const FText FPluginDependenciesGroupNode::GetExtraDisplayName() const
+{
+	if (!bAreChildrenCreated)
+	{
+		return LOCTEXT("DblClickToExpand", "(double click to expand)");
+	}
+	return FTableTreeNode::GetExtraDisplayName();
+}
+
+////////////////////////////////////////////////////////////////////////////////////////////////////
+
+bool FPluginDependenciesGroupNode::OnLazyCreateChildren(TSharedPtr<class UE::Insights::STableTreeView> InTableTreeView)
+{
+	if (bAreChildrenCreated)
+	{
+		return false;
+	}
+
+	TSharedPtr<STreeView<UE::Insights::FTableTreeNodePtr>> TreeView = InTableTreeView->GetInnerTreeView();
+	if (!TreeView || !TreeView->IsItemExpanded(SharedThis(this)))
+	{
+		return false;
+	}
+
+	// [this]
+	// |
+	// +-- [group:{DependentPlugin1}] (self + dependencies) // FPluginAndDependenciesGroupNode
+	// |   |
+	// |   +-- [group:Plugin Dependencies] (double click to expand) // FPluginDependenciesGroupNode, lazy
+	// |   |
+	// |   +-- [group:{DependentPlugin1}] (self) // FPluginSimpleGroupNode
+	// |       |
+	// |       +-- [asset:{Asset1a}]
+	// |       |
+	// |       +-- [asset:{Asset1b}]
+	// |       ...
+	// |
+	// +-- [group:{DependentPlugin2}] (self, no further dependencies) // FPluginSimpleGroupNode
+	// |   |
+	// |   +-- [asset:{Asset2a}]
+	// |   |
+	// |   +-- [asset:{Asset2b}]
+	// |   ...
+	// |
+	// ...
+
+	FAssetTable& AssetTable = GetAssetTableChecked();
+	if (AssetTable.IsValidPluginIndex(PluginIndex))
+	{
+		const FAssetTablePluginInfo& PluginInfo = AssetTable.GetPluginInfoByIndex(PluginIndex);
+
+		// Add dependent plugins.
+		for (int32 DependentPluginIndex : PluginInfo.PluginDependencies)
+		{
+			if (AssetTable.IsValidPluginIndex(DependentPluginIndex))
+			{
+				FName PluginGroupName = AssetTable.GetNameForPlugin(DependentPluginIndex);
+				const FAssetTablePluginInfo& DependentPluginInfo = AssetTable.GetPluginInfoByIndex(DependentPluginIndex);
+				if (DependentPluginInfo.PluginDependencies.Num() > 0)
+				{
+					TSharedPtr<FPluginAndDependenciesGroupNode> PluginGroup = MakeShared<FPluginAndDependenciesGroupNode>(PluginGroupName, GetAssetTableWeak(), DependentPluginIndex);
+					PluginGroup->CreateChildren()->AddAssetChildrenNodes();
+					AddChildAndSetParent(PluginGroup);
+				}
+				else
+				{
+					TSharedPtr<FPluginSimpleGroupNode> PluginGroup = MakeShared<FPluginSimpleGroupNode>(PluginGroupName, GetAssetTableWeak(), DependentPluginIndex);
+					PluginGroup->AddAssetChildrenNodes();
+					AddChildAndSetParent(PluginGroup);
+				}
+			}
+		}
+	}
+
+	bAreChildrenCreated = true;
+	return true;
+}
+
+////////////////////////////////////////////////////////////////////////////////////////////////////
+
+
 ////////////////////////////////////////////////////////////////////////////////////////////////////
 ////////////////////////////////////////////////////////////////////////////////////////////////////
 
diff --git a/Engine/Plugins/Editor/AssetManagerEditor/Source/AssetManagerEditor/Private/TreeView/AssetTreeNode.h b/Engine/Plugins/Editor/AssetManagerEditor/Source/AssetManagerEditor/Private/TreeView/AssetTreeNode.h
--- a/Engine/Plugins/Editor/AssetManagerEditor/Source/AssetManagerEditor/Private/TreeView/AssetTreeNode.h
+++ b/Engine/Plugins/Editor/AssetManagerEditor/Source/AssetManagerEditor/Private/TreeView/AssetTreeNode.h
@@ -80,11 +80,79 @@ public:
 
 	virtual ~FAssetDependenciesGroupTreeNode() {}
 
+	virtual const FSlateBrush* GetIcon() const override;
+	virtual FLinearColor GetColor() const override;
+
 	virtual const FText GetExtraDisplayName() const;
+	virtual bool OnLazyCreateChildren(TSharedPtr<class UE::Insights::STableTreeView> InTableTreeView) override;
+
+private:
+	bool bAreChildrenCreated;
+};
+
+/** Type definition for shared pointers to instances of FAssetDependenciesGroupTreeNode. */
+typedef TSharedPtr<class FAssetDependenciesGroupTreeNode> FAssetDependenciesGroupTreeNodePtr;
+
+////////////////////////////////////////////////////////////////////////////////////////////////////
+
+class FPluginSimpleGroupNode : public FAssetTreeNode
+{
+	INSIGHTS_DECLARE_RTTI(FPluginSimpleGroupNode, FAssetTreeNode)
+
+public:
+	/** Initialization constructor. */
+	explicit FPluginSimpleGroupNode(const FName InGroupName, TWeakPtr<FAssetTable> InParentTable, int32 InPluginIndex)
+		: FAssetTreeNode(InGroupName, InParentTable)
+		, PluginIndex(InPluginIndex)
+	{
+	}
+
+	virtual ~FPluginSimpleGroupNode() {}
 
 	virtual const FSlateBrush* GetIcon() const override;
 	virtual FLinearColor GetColor() const override;
 
+	void AddAssetChildrenNodes();
+
+protected:
+	int32 PluginIndex;
+};
+
+////////////////////////////////////////////////////////////////////////////////////////////////////
+
+class FPluginAndDependenciesGroupNode : public FPluginSimpleGroupNode
+{
+	INSIGHTS_DECLARE_RTTI(FPluginAndDependenciesGroupNode, FPluginSimpleGroupNode)
+
+public:
+	/** Initialization constructor. */
+	explicit FPluginAndDependenciesGroupNode(const FName InGroupName, TWeakPtr<FAssetTable> InParentTable, int32 InPluginIndex)
+		: FPluginSimpleGroupNode(InGroupName, InParentTable, InPluginIndex)
+	{
+	}
+
+	TSharedPtr<FPluginSimpleGroupNode> CreateChildren();
+
+	virtual ~FPluginAndDependenciesGroupNode() {}
+};
+
+////////////////////////////////////////////////////////////////////////////////////////////////////
+
+class FPluginDependenciesGroupNode : public FPluginSimpleGroupNode
+{
+	INSIGHTS_DECLARE_RTTI(FPluginDependenciesGroupNode, FPluginSimpleGroupNode)
+
+public:
+	/** Initialization constructor. */
+	explicit FPluginDependenciesGroupNode(const FName InGroupName, TWeakPtr<FAssetTable> InParentTable, int32 InPluginIndex)
+		: FPluginSimpleGroupNode(InGroupName, InParentTable, InPluginIndex)
+		, bAreChildrenCreated(false)
+	{
+	}
+
+	virtual ~FPluginDependenciesGroupNode() {}
+
+	virtual const FText GetExtraDisplayName() const;
 	virtual bool OnLazyCreateChildren(TSharedPtr<class UE::Insights::STableTreeView> InTableTreeView) override;
 
 private:
diff --git a/Engine/Plugins/Editor/AssetManagerEditor/Source/AssetManagerEditor/Private/TreeView/SAssetTableTreeView.cpp b/Engine/Plugins/Editor/AssetManagerEditor/Source/AssetManagerEditor/Private/TreeView/SAssetTableTreeView.cpp
--- a/Engine/Plugins/Editor/AssetManagerEditor/Source/AssetManagerEditor/Private/TreeView/SAssetTableTreeView.cpp
+++ b/Engine/Plugins/Editor/AssetManagerEditor/Source/AssetManagerEditor/Private/TreeView/SAssetTableTreeView.cpp
@@ -22,6 +22,7 @@
 #include "Insights/Table/ViewModels/TableCellValueFormatter.h"
 #include "Insights/Table/ViewModels/TableColumn.h"
 #include "Insights/Table/ViewModels/TreeNodeGrouping.h"
+#include "Interfaces/IPluginManager.h"
 #include "Logging/MessageLog.h"
 #include "Modules/ModuleManager.h"
 #include "SlateOptMacros.h"
@@ -263,12 +264,12 @@ void SAssetTableTreeView::InitAvailableViewPresets()
 			InOutConfigSet.Add({ FAssetTableColumns::PrimaryNameColumnId,                   true, 200.0f });
 			InOutConfigSet.Add({ FAssetTableColumns::StagedCompressedSizeColumnId,          true, 100.0f });
 			InOutConfigSet.Add({ FAssetTableColumns::TotalSizeUniqueDependenciesColumnId,  !true, 100.0f });
-			InOutConfigSet.Add({ FAssetTableColumns::TotalSizeSharedDependenciesColumnId,   !true, 100.0f });
+			InOutConfigSet.Add({ FAssetTableColumns::TotalSizeSharedDependenciesColumnId,  !true, 100.0f });
 			InOutConfigSet.Add({ FAssetTableColumns::TotalSizeExternalDependenciesColumnId,!true, 100.0f });
 			InOutConfigSet.Add({ FAssetTableColumns::TotalUsageCountColumnId,               true, 100.0f });
 			InOutConfigSet.Add({ FAssetTableColumns::ChunksColumnId,                       !true, 200.0f });
 			InOutConfigSet.Add({ FAssetTableColumns::NativeClassColumnId,                   true, 200.0f });
-			InOutConfigSet.Add({ FAssetTableColumns::PluginNameColumnId,				    true, 200.0f });
+			InOutConfigSet.Add({ FAssetTableColumns::PluginNameColumnId,                    true, 200.0f });
 		}
 	};
 	AvailableViewPresets.Add(MakeShared<FDefaultViewPreset>());
@@ -346,16 +347,77 @@ void SAssetTableTreeView::InitAvailableViewPresets()
 			InOutConfigSet.Add({ FAssetTableColumns::PrimaryNameColumnId,                  !true, 200.0f });
 			InOutConfigSet.Add({ FAssetTableColumns::StagedCompressedSizeColumnId,          true, 100.0f });
 			InOutConfigSet.Add({ FAssetTableColumns::TotalSizeUniqueDependenciesColumnId,   true, 100.0f });
-			InOutConfigSet.Add({ FAssetTableColumns::TotalSizeSharedDependenciesColumnId,    true, 100.0f });
+			InOutConfigSet.Add({ FAssetTableColumns::TotalSizeSharedDependenciesColumnId,   true, 100.0f });
 			InOutConfigSet.Add({ FAssetTableColumns::TotalSizeExternalDependenciesColumnId,!true, 100.0f });
 			InOutConfigSet.Add({ FAssetTableColumns::TotalUsageCountColumnId,              !true, 100.0f });
 			InOutConfigSet.Add({ FAssetTableColumns::ChunksColumnId,                       !true, 200.0f });
 			InOutConfigSet.Add({ FAssetTableColumns::NativeClassColumnId,                   true, 200.0f });
-			InOutConfigSet.Add({ FAssetTableColumns::PluginNameColumnId,             true, 200.0f });
+			InOutConfigSet.Add({ FAssetTableColumns::PluginNameColumnId,                    true, 200.0f });
 		}
 	};
 	AvailableViewPresets.Add(MakeShared<FGameFeaturePluginTypeDependencyView>());
 
+	//////////////////////////////////////////////////
+	// Plugin Dependency View
+
+	class FPluginDependencyView : public UE::Insights::ITableTreeViewPreset
+	{
+	public:
+		virtual FText GetName() const override
+		{
+			return LOCTEXT("PluginDepView_PresetName", "Plugin Dependency Analysis");
+		}
+
+		virtual FText GetToolTip() const override
+		{
+			return LOCTEXT("PluginDepView_PresetToolTip", "Plugin Dependency Analysis View\nConfigure the tree view to show a breakdown of assets by Game Feature Plugin, showing also the dependencies between plugins.");
+		}
+		virtual FName GetSortColumn() const override
+		{
+			return UE::Insights::FTable::GetHierarchyColumnId();
+		}
+		virtual EColumnSortMode::Type GetSortMode() const override
+		{
+			return EColumnSortMode::Type::Ascending;
+		}
+		virtual void SetCurrentGroupings(const TArray<TSharedPtr<UE::Insights::FTreeNodeGrouping>>& InAvailableGroupings, TArray<TSharedPtr<UE::Insights::FTreeNodeGrouping>>& InOutCurrentGroupings) const override
+		{
+			InOutCurrentGroupings.Reset();
+
+			check(InAvailableGroupings[0]->Is<UE::Insights::FTreeNodeGroupingFlat>());
+			InOutCurrentGroupings.Add(InAvailableGroupings[0]);
+
+			const TSharedPtr<UE::Insights::FTreeNodeGrouping>* DependencyGrouping = InAvailableGroupings.FindByPredicate(
+				[](TSharedPtr<UE::Insights::FTreeNodeGrouping>& Grouping)
+				{
+					return Grouping->Is<FPluginDependencyGrouping>();
+				});
+			if (DependencyGrouping)
+			{
+				InOutCurrentGroupings.Add(*DependencyGrouping);
+			}
+		}
+		virtual void GetColumnConfigSet(TArray<UE::Insights::FTableColumnConfig>& InOutConfigSet) const override
+		{
+			InOutConfigSet.Add({ UE::Insights::FTable::GetHierarchyColumnId(),              true, 400.0f });
+			InOutConfigSet.Add({ FAssetTableColumns::CountColumnId,                         true, 100.0f });
+			InOutConfigSet.Add({ FAssetTableColumns::TypeColumnId,                          true, 200.0f });
+			InOutConfigSet.Add({ FAssetTableColumns::NameColumnId,                         !true, 200.0f });
+			InOutConfigSet.Add({ FAssetTableColumns::PathColumnId,                         !true, 400.0f });
+			InOutConfigSet.Add({ FAssetTableColumns::PrimaryTypeColumnId,                   true, 200.0f });
+			InOutConfigSet.Add({ FAssetTableColumns::PrimaryNameColumnId,                  !true, 200.0f });
+			InOutConfigSet.Add({ FAssetTableColumns::StagedCompressedSizeColumnId,          true, 100.0f });
+			InOutConfigSet.Add({ FAssetTableColumns::TotalSizeUniqueDependenciesColumnId,   true, 100.0f });
+			InOutConfigSet.Add({ FAssetTableColumns::TotalSizeSharedDependenciesColumnId,   true, 100.0f });
+			InOutConfigSet.Add({ FAssetTableColumns::TotalSizeExternalDependenciesColumnId,!true, 100.0f });
+			InOutConfigSet.Add({ FAssetTableColumns::TotalUsageCountColumnId,              !true, 100.0f });
+			InOutConfigSet.Add({ FAssetTableColumns::ChunksColumnId,                       !true, 200.0f });
+			InOutConfigSet.Add({ FAssetTableColumns::NativeClassColumnId,                   true, 200.0f });
+			InOutConfigSet.Add({ FAssetTableColumns::PluginNameColumnId,                   !true, 200.0f });
+		}
+	};
+	AvailableViewPresets.Add(MakeShared<FPluginDependencyView>());
+
 	//////////////////////////////////////////////////
 	// Path Breakdown View
 
@@ -407,12 +469,12 @@ void SAssetTableTreeView::InitAvailableViewPresets()
 			InOutConfigSet.Add({ FAssetTableColumns::PrimaryNameColumnId,                   true, 200.0f });
 			InOutConfigSet.Add({ FAssetTableColumns::StagedCompressedSizeColumnId,          true, 100.0f });
 			InOutConfigSet.Add({ FAssetTableColumns::TotalSizeUniqueDependenciesColumnId,   true, 100.0f });
-			InOutConfigSet.Add({ FAssetTableColumns::TotalSizeSharedDependenciesColumnId,    true, 100.0f });
+			InOutConfigSet.Add({ FAssetTableColumns::TotalSizeSharedDependenciesColumnId,   true, 100.0f });
 			InOutConfigSet.Add({ FAssetTableColumns::TotalSizeExternalDependenciesColumnId,!true, 100.0f });
 			InOutConfigSet.Add({ FAssetTableColumns::TotalUsageCountColumnId,               true, 100.0f });
 			InOutConfigSet.Add({ FAssetTableColumns::ChunksColumnId,                       !true, 200.0f });
 			InOutConfigSet.Add({ FAssetTableColumns::NativeClassColumnId,                   true, 200.0f });
-			InOutConfigSet.Add({ FAssetTableColumns::PluginNameColumnId,             true, 200.0f });
+			InOutConfigSet.Add({ FAssetTableColumns::PluginNameColumnId,                    true, 200.0f });
 		}
 	};
 	AvailableViewPresets.Add(MakeShared<FAssetPathViewPreset>());
@@ -468,12 +530,12 @@ void SAssetTableTreeView::InitAvailableViewPresets()
 			InOutConfigSet.Add({ FAssetTableColumns::PrimaryNameColumnId,                  !true, 200.0f });
 			InOutConfigSet.Add({ FAssetTableColumns::StagedCompressedSizeColumnId,          true, 100.0f });
 			InOutConfigSet.Add({ FAssetTableColumns::TotalSizeUniqueDependenciesColumnId,   true, 100.0f });
-			InOutConfigSet.Add({ FAssetTableColumns::TotalSizeSharedDependenciesColumnId,    true, 100.0f });
+			InOutConfigSet.Add({ FAssetTableColumns::TotalSizeSharedDependenciesColumnId,   true, 100.0f });
 			InOutConfigSet.Add({ FAssetTableColumns::TotalSizeExternalDependenciesColumnId,!true, 100.0f });
 			InOutConfigSet.Add({ FAssetTableColumns::TotalUsageCountColumnId,               true, 100.0f });
 			InOutConfigSet.Add({ FAssetTableColumns::ChunksColumnId,                       !true, 200.0f });
 			InOutConfigSet.Add({ FAssetTableColumns::NativeClassColumnId,                   true, 200.0f });
-			InOutConfigSet.Add({ FAssetTableColumns::PluginNameColumnId,             true, 200.0f });
+			InOutConfigSet.Add({ FAssetTableColumns::PluginNameColumnId,                    true, 200.0f });
 		}
 	};
 	AvailableViewPresets.Add(MakeShared<FPrimaryAssetViewPreset>());
@@ -564,6 +626,7 @@ void SAssetTableTreeView::InternalCreateGroupings()
 	int32 Index = 1; // after the Flat ("All") grouping
 
 	AvailableGroupings.Insert(MakeShared<FAssetDependencyGrouping>(), Index++);
+	AvailableGroupings.Insert(MakeShared<FPluginDependencyGrouping>(), Index++);
 }
 
 ////////////////////////////////////////////////////////////////////////////////////////////////////
@@ -1177,6 +1240,8 @@ void SAssetTableTreeView::RefreshAssets()
 	// Now is safe to clear the previous assets.
 	AssetTable->ClearAllData();
 
+	TMap<FString, int64> PluginToSizeMap;
+
 	if (IsRegistrySourceValid())
 	{
 		TMap<FAssetData, int32> AssetToIndexMap;
@@ -1198,6 +1263,15 @@ void SAssetTableTreeView::RefreshAssets()
 					PopulateAssetTableRow(AssetRow, SourceAsset, *AssetTable);
 					AssetTable->AddAsset(AssetRow);
 					AssetToIndexMap.Add(SourceAsset, AssetTable->GetTotalAssetCount() - 1);
+
+					if (int64* PluginSize = PluginToSizeMap.Find(AssetRow.GetPluginName()))
+					{
+						(*PluginSize) += AssetRow.GetStagedCompressedSize();
+					}
+					else
+					{
+						PluginToSizeMap.Add(AssetRow.GetPluginName(), AssetRow.GetStagedCompressedSize());
+					}
 				}
 			}
 		}
@@ -1250,6 +1324,14 @@ void SAssetTableTreeView::RefreshAssets()
 						IndicesOfDependenciesInRowTableToAddToCurrentSourceAssetRow.Add(AssetTable->GetTotalAssetCount() - 1);
 						AssetToIndexMap.Add(*DependencyAsset, AssetTable->GetTotalAssetCount() - 1);
 						SourceAssets.Add(*DependencyAsset);
+						if (int64* PluginSize = PluginToSizeMap.Find(NewRow.GetPluginName()))
+						{
+							(*PluginSize) += NewRow.GetStagedCompressedSize();
+						}
+						else
+						{
+							PluginToSizeMap.Add(NewRow.GetPluginName(), NewRow.GetStagedCompressedSize());
+						}
 					}
 				}
 			}
@@ -1293,6 +1375,30 @@ void SAssetTableTreeView::RefreshAssets()
 		}
 	}
 
+	// Setup plugin infos and dependencies. This will eventually be replaced by data from the asset registry
+	for (TPair<FString, int64>& PluginEntry : PluginToSizeMap)
+	{
+		if (TSharedPtr<IPlugin> Plugin = IPluginManager::Get().FindPlugin(PluginEntry.Key))
+		{
+			const TCHAR* StoredPluginName = AssetTable->StoreStr(PluginEntry.Key);
+			FAssetTablePluginInfo& PluginInfo = AssetTable->GetOrCreatePluginInfo(StoredPluginName);
+			PluginInfo.Size = PluginEntry.Value;
+			TArray<FPluginReferenceDescriptor> PluginReferences;
+			IPluginManager::Get().GetPluginDependencies(PluginEntry.Key, PluginReferences);
+			for (const FPluginReferenceDescriptor& ReferenceDescriptor : PluginReferences)
+			{
+				const TCHAR* StoredReferencePluginName = AssetTable->StoreStr(ReferenceDescriptor.Name);
+				int32 DependencyIndex = AssetTable->GetIndexForPlugin(StoredReferencePluginName);
+				if (DependencyIndex == -1)
+				{
+					DependencyIndex = AssetTable->GetNumPlugins();
+					AssetTable->GetOrCreatePluginInfo(StoredReferencePluginName);
+				}
+				PluginInfo.PluginDependencies.Add(DependencyIndex);
+			}
+		}
+	}
+
 	Stopwatch.Stop();
 	const double TotalTime = Stopwatch.GetAccumulatedTime();
 	UE_LOG(LogInsights, Log, TEXT("[AssetTree] Asset table rebuilt in %.4fs (%d visible + %d hidden = %d assets)"),
