Added the ability to have two disk size tool windows open at once, added some additional text to the window, and improved the handling of the situation where you start opening a registry and cancel the file open dialog

diff --git a/Engine/Plugins/Editor/AssetManagerEditor/Source/AssetManagerEditor/Private/AssetManagerEditorModule.cpp b/Engine/Plugins/Editor/AssetManagerEditor/Source/AssetManagerEditor/Private/AssetManagerEditorModule.cpp
--- a/Engine/Plugins/Editor/AssetManagerEditor/Source/AssetManagerEditor/Private/AssetManagerEditorModule.cpp
+++ b/Engine/Plugins/Editor/AssetManagerEditor/Source/AssetManagerEditor/Private/AssetManagerEditorModule.cpp
@@ -362,6 +362,7 @@ private:
 	static const TCHAR* FindClassDepHelpText;
 	static const FName AssetAuditTabName;
 	static const FName AssetDiskSizeTabName;
+	static const FName AssetDiskSize2TabName;
 	static const FName ReferenceViewerTabName;
 	static const FName SizeMapTabName;
 
@@ -413,6 +414,7 @@ const TCHAR* FAssetManagerEditorModule::FindDepChainHelpText = TEXT("Finds all d
 const TCHAR* FAssetManagerEditorModule::FindClassDepHelpText = TEXT("Finds all dependencies of a certain set of classes to the target asset.\n Usage: FindDepClasses TargetPackagePath ClassName1 ClassName2 etc (optional: -hardonly/-softonly) \n e.g. FindDepChain /game/characters/heroes/muriel/meshes/muriel /game/cards");
 const FName FAssetManagerEditorModule::AssetAuditTabName = TEXT("AssetAudit");
 const FName FAssetManagerEditorModule::AssetDiskSizeTabName = TEXT("AssetDiskSize");
+const FName FAssetManagerEditorModule::AssetDiskSize2TabName = TEXT("AssetDiskSize2");
 const FName FAssetManagerEditorModule::ReferenceViewerTabName = TEXT("ReferenceViewer");
 const FName FAssetManagerEditorModule::SizeMapTabName = TEXT("SizeMap");
 
@@ -518,12 +520,19 @@ void FAssetManagerEditorModule::StartupModule()
 		FGlobalTabmanager::Get()->RegisterDefaultTabWindowSize(AssetAuditTabName, FVector2D(1080, 600));
 
 		FGlobalTabmanager::Get()->RegisterNomadTabSpawner(AssetDiskSizeTabName, FOnSpawnTab::CreateRaw(this, &FAssetManagerEditorModule::SpawnAssetDiskSizeTab))
-			.SetDisplayName(LOCTEXT("AssetDiskSizeTitle", "Asset Disk Size"))
+			.SetDisplayName(LOCTEXT("AssetDiskSizeTitle", "Asset Disk Size 1"))
 			.SetTooltipText(LOCTEXT("AssetDiskSizeTooltip", "Open Asset Disk Size window, allows to analyze disk size for assets."))
 			.SetGroup(WorkspaceMenu::GetMenuStructure().GetDeveloperToolsAuditCategory())
 			.SetIcon(FSlateIcon(FAppStyle::GetAppStyleSetName(), "LevelEditor.Audit"));
 		FGlobalTabmanager::Get()->RegisterDefaultTabWindowSize(AssetDiskSizeTabName, FVector2D(1080, 600));
 
+		FGlobalTabmanager::Get()->RegisterNomadTabSpawner(AssetDiskSize2TabName, FOnSpawnTab::CreateRaw(this, &FAssetManagerEditorModule::SpawnAssetDiskSizeTab))
+			.SetDisplayName(LOCTEXT("AssetDiskSizeTitle", "Asset Disk Size 2"))
+			.SetTooltipText(LOCTEXT("AssetDiskSizeTooltip", "Open Asset Disk Size window, allows to analyze disk size for assets."))
+			.SetGroup(WorkspaceMenu::GetMenuStructure().GetDeveloperToolsAuditCategory())
+			.SetIcon(FSlateIcon(FAppStyle::GetAppStyleSetName(), "LevelEditor.Audit"));
+		FGlobalTabmanager::Get()->RegisterDefaultTabWindowSize(AssetDiskSize2TabName, FVector2D(1080, 600));
+
 		FGlobalTabmanager::Get()->RegisterNomadTabSpawner(ReferenceViewerTabName, FOnSpawnTab::CreateRaw(this, &FAssetManagerEditorModule::SpawnReferenceViewerTab))
 			.SetDisplayName(LOCTEXT("ReferenceViewerTitle", "Reference Viewer"))
 			.SetMenuType(ETabSpawnerMenuType::Hidden);
diff --git a/Engine/Plugins/Editor/AssetManagerEditor/Source/AssetManagerEditor/Private/TreeView/AssetTable.h b/Engine/Plugins/Editor/AssetManagerEditor/Source/AssetManagerEditor/Private/TreeView/AssetTable.h
--- a/Engine/Plugins/Editor/AssetManagerEditor/Source/AssetManagerEditor/Private/TreeView/AssetTable.h
+++ b/Engine/Plugins/Editor/AssetManagerEditor/Source/AssetManagerEditor/Private/TreeView/AssetTable.h
@@ -136,6 +136,7 @@ struct FAssetTableDependencySizes
 struct FAssetTablePluginInfo
 {
 	TSet<int32> PluginDependencies;
+	TSet<int32> DiscoveredPluginDependencies;
 	const TCHAR* PluginName = nullptr;
 	int64 Size = -1;
 };
diff --git a/Engine/Plugins/Editor/AssetManagerEditor/Source/AssetManagerEditor/Private/TreeView/SAssetTableTreeView.cpp b/Engine/Plugins/Editor/AssetManagerEditor/Source/AssetManagerEditor/Private/TreeView/SAssetTableTreeView.cpp
--- a/Engine/Plugins/Editor/AssetManagerEditor/Source/AssetManagerEditor/Private/TreeView/SAssetTableTreeView.cpp
+++ b/Engine/Plugins/Editor/AssetManagerEditor/Source/AssetManagerEditor/Private/TreeView/SAssetTableTreeView.cpp
@@ -72,6 +72,8 @@ void SAssetTableTreeView::Construct(const FArguments& InArgs, TSharedPtr<FAssetT
 	CreateGroupings();
 	CreateSortings();
 
+	RegistrySourceTimeText->SetText(LOCTEXT("RegistrySourceTimeText_None", "No registry loaded."));
+
 	RequestOpenRegistry();
 }
 
@@ -175,6 +177,21 @@ void SAssetTableTreeView::ConstructHeaderArea(TSharedRef<SVerticalBox> InWidgetC
 			]
 		];
 
+	InWidgetContent->AddSlot()
+		.VAlign(VAlign_Center)
+		.AutoHeight()
+		.Padding(2.0f)
+		[
+			SNew(SHorizontalBox)
+
+			+ SHorizontalBox::Slot()
+			.FillWidth(1.0f)
+			.VAlign(VAlign_Center)
+			[
+				SAssignNew(RegistrySourceTimeText, STextBlock)
+			]
+		];
+
 	InWidgetContent->AddSlot()
 		.VAlign(VAlign_Center)
 		.AutoHeight()
@@ -870,6 +887,7 @@ void SAssetTableTreeView::ExtendMenu(FMenuBuilder& MenuBuilder)
 void SAssetTableTreeView::RequestOpenRegistry()
 {
 	bNeedsToOpenRegistry = true;
+	FooterLeftTextStoredPreOpen = FooterLeftText;
 	FooterLeftText = LOCTEXT("FooterLeftTextFmt_OpenRegistry_Filtered", "Opening registry... please wait...");
 }
 
@@ -880,8 +898,23 @@ void SAssetTableTreeView::OpenRegistry()
 	RegistrySource.SourceName = FAssetManagerEditorRegistrySource::CustomSourceName;
 	if (EditorModule->PopulateRegistrySource(&RegistrySource))
 	{
+		IFileManager& FileManager = IFileManager::Get();
+		const FString RegistryFilePath = FileManager.GetFilenameOnDisk(*FileManager.ConvertToAbsolutePathForExternalAppForRead(*RegistrySource.SourceFilename));
+		RegistrySourceTimeText->SetText(FText::Format(LOCTEXT("RegistrySourceTimeText", "Loaded: {0} from {1}"), 
+			FText::FromString(RegistryFilePath),
+			FText::FromString(RegistrySource.SourceTimestamp)));
+
 		RequestRefreshAssets();
 	}
+	else if (!IsRegistrySourceValid())
+	{
+		FooterLeftText = LOCTEXT("FooterLeftTextFmt_OpenRegistry_Failed", "No registry selected or load failed.");
+		RegistrySourceTimeText->SetText(LOCTEXT("RegistrySourceTimeText_None", "No registry loaded."));
+	}
+	else
+	{
+		FooterLeftText = FooterLeftTextStoredPreOpen;
+	}
 }
 
 ////////////////////////////////////////////////////////////////////////////////////////////////////
@@ -1242,6 +1275,9 @@ void SAssetTableTreeView::RefreshAssets()
 
 	TMap<FString, int64> PluginToSizeMap;
 
+	typedef TSet<const TCHAR*, TStringPointerSetKeyFuncs_DEPRECATED<const TCHAR*>> DeprecatedTCharSetType;
+	TMap<const TCHAR*, DeprecatedTCharSetType, FDefaultSetAllocator, TStringPointerMapKeyFuncs_DEPRECATED<const TCHAR*, DeprecatedTCharSetType>> DiscoveredPluginDependencyEdges;
+
 	if (IsRegistrySourceValid())
 	{
 		TMap<FAssetData, int32> AssetToIndexMap;
@@ -1348,9 +1384,25 @@ void SAssetTableTreeView::RefreshAssets()
 					}
 					if (ensure(RowIndex != nullptr))
 					{
+						const TCHAR* CurrentPlugin = AssetTable->GetAssetChecked(*RowIndex).GetPluginName();
+						DeprecatedTCharSetType DependencyPlugins;
+
 						for (int32 Index : IndicesOfDependenciesInRowTableToAddToCurrentSourceAssetRow)
 						{
 							AssetTable->GetAsset(*RowIndex)->Dependencies.AddUnique(Index);
+
+							const TCHAR* DependencyPlugin = AssetTable->GetAsset(Index)->GetPluginName();
+							if (DependencyPlugin != CurrentPlugin)
+							{
+								DependencyPlugins.Add(DependencyPlugin);
+							}
+						}
+
+
+						DeprecatedTCharSetType& DiscoveredPluginDependencyList = DiscoveredPluginDependencyEdges.FindOrAdd(CurrentPlugin);
+						for (const TCHAR* DependencyPlugin : DependencyPlugins)
+						{
+							DiscoveredPluginDependencyList.Add(DependencyPlugin);
 						}
 					}
 				}
@@ -1378,11 +1430,11 @@ void SAssetTableTreeView::RefreshAssets()
 	// Setup plugin infos and dependencies. This will eventually be replaced by data from the asset registry
 	for (TPair<FString, int64>& PluginEntry : PluginToSizeMap)
 	{
+		const TCHAR* StoredPluginName = AssetTable->StoreStr(PluginEntry.Key);
+		FAssetTablePluginInfo& PluginInfo = AssetTable->GetOrCreatePluginInfo(StoredPluginName);
+		PluginInfo.Size = PluginEntry.Value;
 		if (TSharedPtr<IPlugin> Plugin = IPluginManager::Get().FindPlugin(PluginEntry.Key))
 		{
-			const TCHAR* StoredPluginName = AssetTable->StoreStr(PluginEntry.Key);
-			FAssetTablePluginInfo& PluginInfo = AssetTable->GetOrCreatePluginInfo(StoredPluginName);
-			PluginInfo.Size = PluginEntry.Value;
 			TArray<FPluginReferenceDescriptor> PluginReferences;
 			IPluginManager::Get().GetPluginDependencies(PluginEntry.Key, PluginReferences);
 			for (const FPluginReferenceDescriptor& ReferenceDescriptor : PluginReferences)
@@ -1394,11 +1446,58 @@ void SAssetTableTreeView::RefreshAssets()
 					DependencyIndex = AssetTable->GetNumPlugins();
 					AssetTable->GetOrCreatePluginInfo(StoredReferencePluginName);
 				}
-				PluginInfo.PluginDependencies.Add(DependencyIndex);
+				// Note that we can't use PluginInfo directly because the above code could have grown the array 
+				// and that would invalidate the reference
+				AssetTable->GetOrCreatePluginInfo(StoredPluginName).PluginDependencies.Add(DependencyIndex);
 			}
 		}
+		else
+		{
+			UE_LOG(LogInsights, Warning, TEXT("Could not find plugin %s in plugin manager."), PluginInfo.PluginName);
+		}
+	}
+
+	// Discovered version
+	for (const TPair<const TCHAR*, DeprecatedTCharSetType>& DiscoveredDependencies : DiscoveredPluginDependencyEdges)
+	{
+		const TCHAR* StoredPluginName = AssetTable->StoreStr(DiscoveredDependencies.Key);
+		FAssetTablePluginInfo& PluginInfo = AssetTable->GetOrCreatePluginInfo(StoredPluginName);
+
+		for (const TCHAR* DependencyName : DiscoveredDependencies.Value)
+		{
+			int32 DependencyIndex = AssetTable->GetIndexForPlugin(DependencyName);
+			ensureAlways(DependencyIndex != -1); // These were created above.
+			PluginInfo.DiscoveredPluginDependencies.Add(DependencyIndex);
+		}
 	}
 
+	// TODO: Once we have the dependency data from the uplugin files in the cook we can 
+	// compare them to the discovered dependencies to see if there are plugins we're dependending on
+	// that, perhaps, we shouldn't be. That said, we shouldn't storer the discovered dependencies on the 
+	// plugin infos at that point.
+
+	//// Analyze difference
+	//UE_LOG(LogInsights, Warning, TEXT("Registry Deps-------"));
+	//for (int32 PluginIndex = 0; PluginIndex < AssetTable->GetNumPlugins(); PluginIndex++)
+	//{
+	//	const FAssetTablePluginInfo& PluginInfo = AssetTable->GetPluginInfoByIndex(PluginIndex);
+	//	UE_LOG(LogInsights, Warning, TEXT("Plugin: %s"), PluginInfo.PluginName);
+	//	for (int32 DependencyIndex : PluginInfo.PluginDependencies)
+	//	{
+	//		UE_LOG(LogInsights, Warning, TEXT("\t%s"), AssetTable->GetPluginInfoByIndex(DependencyIndex).PluginName);
+	//	}
+	//}
+	//UE_LOG(LogInsights, Warning, TEXT("Discovered Deps-------"));
+	//for (int32 PluginIndex = 0; PluginIndex < AssetTable->GetNumPlugins(); PluginIndex++)
+	//{
+	//	const FAssetTablePluginInfo& PluginInfo = AssetTable->GetPluginInfoByIndex(PluginIndex);
+	//	UE_LOG(LogInsights, Warning, TEXT("Plugin: %s"), PluginInfo.PluginName);
+	//	for (int32 DependencyIndex : PluginInfo.DiscoveredPluginDependencies)
+	//	{
+	//		UE_LOG(LogInsights, Warning, TEXT("\t%s"), AssetTable->GetPluginInfoByIndex(DependencyIndex).PluginName);
+	//	}
+	//}
+
 	Stopwatch.Stop();
 	const double TotalTime = Stopwatch.GetAccumulatedTime();
 	UE_LOG(LogInsights, Log, TEXT("[AssetTree] Asset table rebuilt in %.4fs (%d visible + %d hidden = %d assets)"),
diff --git a/Engine/Plugins/Editor/AssetManagerEditor/Source/AssetManagerEditor/Private/TreeView/SAssetTableTreeView.h b/Engine/Plugins/Editor/AssetManagerEditor/Source/AssetManagerEditor/Private/TreeView/SAssetTableTreeView.h
--- a/Engine/Plugins/Editor/AssetManagerEditor/Source/AssetManagerEditor/Private/TreeView/SAssetTableTreeView.h
+++ b/Engine/Plugins/Editor/AssetManagerEditor/Source/AssetManagerEditor/Private/TreeView/SAssetTableTreeView.h
@@ -116,6 +116,8 @@ private:
 	bool bNeedsToRefreshAssets = false;
 	bool bNeedsToRebuild = false;
 
+	FText FooterLeftTextStoredPreOpen;
+	FText RegistryLoadedText;
 	FText FooterLeftText;
 	FText FooterCenterText1;
 	FText FooterCenterText2;
