Made a number of UX changes to the Asset Disk Size Tool for usability
Disabled color by asset type ( should be optional in future )
Renamed Plugin Sub Folder groups to Assets and Dependencies
Changed the Icon for plugins so they distinguish from regular folders.

diff --git a/Engine/Plugins/Editor/AssetManagerEditor/Source/AssetManagerEditor/Private/Insights/Common/InsightsStyle.cpp b/Engine/Plugins/Editor/AssetManagerEditor/Source/AssetManagerEditor/Private/Insights/Common/InsightsStyle.cpp
--- a/Engine/Plugins/Editor/AssetManagerEditor/Source/AssetManagerEditor/Private/Insights/Common/InsightsStyle.cpp
+++ b/Engine/Plugins/Editor/AssetManagerEditor/Source/AssetManagerEditor/Private/Insights/Common/InsightsStyle.cpp
@@ -165,6 +165,7 @@ void FInsightsStyle::FStyle::Initialize()
 	Set("Icons.Leaf.TreeItem", new CORE_IMAGE_BRUSH_SVG("Starship/Common/bullet-point", Icon12x12));
 	Set("Icons.Asset.TreeItem", new CORE_IMAGE_BRUSH_SVG("Starship/Common/box-perspective", Icon12x12));
 	Set("Icons.Package.TreeItem", new EDITOR_IMAGE_BRUSH_SVG("Starship/Common/ProjectPackage", Icon12x12));
+	Set("Icons.Plugin.TreeItem", new EDITOR_IMAGE_BRUSH_SVG("Starship/Common/Plugins", Icon12x12));
 
 	Set("TreeTable.TooltipBold", FTextBlockStyle(NormalText)
 		.SetFont(DEFAULT_FONT("Bold", 8))
diff --git a/Engine/Plugins/Editor/AssetManagerEditor/Source/AssetManagerEditor/Private/TreeView/AssetTreeNode.cpp b/Engine/Plugins/Editor/AssetManagerEditor/Source/AssetManagerEditor/Private/TreeView/AssetTreeNode.cpp
--- a/Engine/Plugins/Editor/AssetManagerEditor/Source/AssetManagerEditor/Private/TreeView/AssetTreeNode.cpp
+++ b/Engine/Plugins/Editor/AssetManagerEditor/Source/AssetManagerEditor/Private/TreeView/AssetTreeNode.cpp
@@ -20,36 +20,88 @@ INSIGHTS_IMPLEMENT_RTTI(FPluginDependenciesGroupNode)
 // FAssetTreeNode
 ////////////////////////////////////////////////////////////////////////////////////////////////////
 
-const FSlateBrush* FAssetTreeNode::GetIcon() const
+const FSlateBrush* FAssetTreeNode::GetIcon(EStyle Style) const
 {
-	if (IsValidAsset())
+	switch ( Style )
 	{
-		return FBaseTreeNode::GetDefaultIcon(false); // default icon for leaf nodes
+		default:
+		case EStyle::EDefault:
+		case EStyle::EAsset:
+		{
+			return UE::Insights::FInsightsStyle::GetBrush("Icons.Leaf.TreeItem");
+			break;
+		}
+				
+		case EStyle::EGroup:
+		{
+			return UE::Insights::FInsightsStyle::GetBrush("Icons.Group.TreeItem");
+			break;
+		}	
+
+		case EStyle::EPlugin:
+		{
+			return UE::Insights::FInsightsStyle::GetBrush("Icons.Plugin.TreeItem");
+			break;
+		}
 	}
-	return FBaseTreeNode::GetDefaultIcon(IsGroup()); // default icon
 }
 
-////////////////////////////////////////////////////////////////////////////////////////////////////
+FLinearColor FAssetTreeNode::GetColor(EStyle Style) const 
+{
+	switch (Style)
+	{
+		default:
+		case EStyle::EDefault:
+		{
+			return FLinearColor(1.0f, 1.0f, 0.5f, 1.0f);
+			break;
+		}
 
-FLinearColor FAssetTreeNode::GetColor() const
+		case EStyle::EAsset:
+		{
+			//return GetAssetChecked().GetColor();
+			return FLinearColor(1.0f, 1.0f, 1.0f, 1.0f);
+			break;
+		}
+
+		case EStyle::EGroup:
+		{
+			return FLinearColor(1.0f, 1.0f, 0.5f, 1.0f);
+			break;
+		}
+
+		case EStyle::EPlugin:
+		{
+			return FLinearColor(1.0f, 1.0f, 1.0f, 1.0f);
+			break;
+		}	
+	}	
+}
+
+FAssetTreeNode::EStyle FAssetTreeNode::GetStyle() const
 {
-	return GetAssetChecked().GetColor();
+	return IsValidAsset() ? EStyle::EAsset : ( IsGroup() ? EStyle::EGroup : EStyle::EDefault );
+}
+
+const FSlateBrush* FAssetTreeNode::GetIcon() const
+{
+	return GetIcon(GetStyle());
 }
 
-////////////////////////////////////////////////////////////////////////////////////////////////////
-// FAssetDependenciesGroupTreeNode
 ////////////////////////////////////////////////////////////////////////////////////////////////////
 
-const FSlateBrush* FAssetDependenciesGroupTreeNode::GetIcon() const
+FLinearColor FAssetTreeNode::GetColor() const
 {
-	return FBaseTreeNode::GetDefaultIcon(true); // default icon for group nodes
+	return GetColor(GetStyle());
 }
 
+////////////////////////////////////////////////////////////////////////////////////////////////////
+// FAssetDependenciesGroupTreeNode
 ////////////////////////////////////////////////////////////////////////////////////////////////////
 
-FLinearColor FAssetDependenciesGroupTreeNode::GetColor() const
+FAssetTreeNode::EStyle FAssetDependenciesGroupTreeNode::GetStyle() const
 {
-	return FLinearColor(0.75f, 0.5f, 1.0f, 1.0f);
+	return EStyle::EGroup;
 }
 
 ////////////////////////////////////////////////////////////////////////////////////////////////////
@@ -106,18 +158,10 @@ bool FAssetDependenciesGroupTreeNode::OnLazyCreateChildren(TSharedPtr<UE::Insigh
 // FPluginSimpleGroupNode
 ////////////////////////////////////////////////////////////////////////////////////////////////////
 
-const FSlateBrush* FPluginSimpleGroupNode::GetIcon() const
-{
-	return UE::Insights::FBaseTreeNode::GetDefaultIcon(true);
-}
-
-////////////////////////////////////////////////////////////////////////////////////////////////////
-
-FLinearColor FPluginSimpleGroupNode::GetColor() const
+FAssetTreeNode::EStyle FPluginSimpleGroupNode::GetStyle() const
 {
-	return FLinearColor(1.0f, 1.0f, 1.0f, 1.0f);
+	return EStyle::EGroup;
 }
-
 ////////////////////////////////////////////////////////////////////////////////////////////////////
 // FPluginSimpleGroupNode
 ////////////////////////////////////////////////////////////////////////////////////////////////////
@@ -148,6 +192,11 @@ void FPluginSimpleGroupNode::AddAssetChildrenNodes()
 // FPluginAndDependenciesGroupNode
 ////////////////////////////////////////////////////////////////////////////////////////////////////
 
+FAssetTreeNode::EStyle FPluginAndDependenciesGroupNode::GetStyle() const
+{
+	return EStyle::EPlugin;
+}
+
 TSharedPtr<FPluginSimpleGroupNode> FPluginAndDependenciesGroupNode::CreateChildren()
 {
 	// [this]
@@ -161,12 +210,13 @@ TSharedPtr<FPluginSimpleGroupNode> FPluginAndDependenciesGroupNode::CreateChildr
 	{
 		// Create the Plugin Dependencies group node.
 		// The children nodes (list of dependent plugins) will be lazy created.
-		static FName DependenciesGroupName(TEXT("Plugin Dependencies"));
+		static FName DependenciesGroupName(TEXT("Dependencies"));
 		TSharedPtr<FPluginDependenciesGroupNode> DependenciesGroup = MakeShared<FPluginDependenciesGroupNode>(DependenciesGroupName, GetAssetTableWeak(), PluginIndex);
 		AddChildAndSetParent(DependenciesGroup);
 
 		// Create the Plugin Self group node (where asset nodes will be added).
-		FName PluginGroupName = AssetTable.GetNameForPlugin(PluginIndex);
+		//FName PluginGroupName = AssetTable.GetNameForPlugin(PluginIndex);
+		static FName PluginGroupName(TEXT("Assets"));
 		TSharedPtr<FPluginSimpleGroupNode> PluginGroup = MakeShared<FPluginSimpleGroupNode>(PluginGroupName, GetAssetTableWeak(), PluginIndex);
 		AddChildAndSetParent(PluginGroup);
 		return PluginGroup;
@@ -256,10 +306,11 @@ bool FPluginDependenciesGroupNode::OnLazyCreateChildren(TSharedPtr<class UE::Ins
 	return true;
 }
 
-////////////////////////////////////////////////////////////////////////////////////////////////////
-
+FAssetTreeNode::EStyle FPluginDependenciesGroupNode::GetStyle() const
+{
+	return EStyle::EGroup;
+}
 
-////////////////////////////////////////////////////////////////////////////////////////////////////
 ////////////////////////////////////////////////////////////////////////////////////////////////////
 
 #undef LOCTEXT_NAMESPACE
diff --git a/Engine/Plugins/Editor/AssetManagerEditor/Source/AssetManagerEditor/Private/TreeView/AssetTreeNode.h b/Engine/Plugins/Editor/AssetManagerEditor/Source/AssetManagerEditor/Private/TreeView/AssetTreeNode.h
--- a/Engine/Plugins/Editor/AssetManagerEditor/Source/AssetManagerEditor/Private/TreeView/AssetTreeNode.h
+++ b/Engine/Plugins/Editor/AssetManagerEditor/Source/AssetManagerEditor/Private/TreeView/AssetTreeNode.h
@@ -43,8 +43,23 @@ public:
 	FAssetTable& GetAssetTableChecked() const { return *AssetTablePtr; }
 	const FAssetTableRow& GetAssetChecked() const { return AssetTablePtr->GetAssetChecked(RowId.RowIndex); }
 
-	virtual const FSlateBrush* GetIcon() const override;
-	virtual FLinearColor GetColor() const override;
+	virtual const FSlateBrush* GetIcon() const final;
+	virtual FLinearColor GetColor() const final;
+	
+protected:
+
+	// Set of UI style options for node types
+	enum class EStyle
+	{
+		EDefault,
+		EGroup,
+		EAsset,
+		EPlugin,
+	};
+
+	virtual EStyle GetStyle() const;
+	const FSlateBrush* GetIcon(EStyle Style) const;
+	FLinearColor GetColor(EStyle Style) const;
 
 private:
 	FAssetTable* AssetTablePtr;
@@ -80,9 +95,8 @@ public:
 
 	virtual ~FAssetDependenciesGroupTreeNode() {}
 
-	virtual const FSlateBrush* GetIcon() const override;
-	virtual FLinearColor GetColor() const override;
-
+	virtual EStyle GetStyle() const override;
+	
 	virtual const FText GetExtraDisplayName() const;
 	virtual bool OnLazyCreateChildren(TSharedPtr<class UE::Insights::STableTreeView> InTableTreeView) override;
 
@@ -109,8 +123,7 @@ public:
 
 	virtual ~FPluginSimpleGroupNode() {}
 
-	virtual const FSlateBrush* GetIcon() const override;
-	virtual FLinearColor GetColor() const override;
+	virtual EStyle GetStyle() const override;
 
 	void AddAssetChildrenNodes();
 
@@ -133,6 +146,8 @@ public:
 
 	TSharedPtr<FPluginSimpleGroupNode> CreateChildren();
 
+	virtual EStyle GetStyle() const override;
+
 	virtual ~FPluginAndDependenciesGroupNode() {}
 };
 
@@ -155,6 +170,8 @@ public:
 	virtual const FText GetExtraDisplayName() const;
 	virtual bool OnLazyCreateChildren(TSharedPtr<class UE::Insights::STableTreeView> InTableTreeView) override;
 
+	virtual EStyle GetStyle() const override;
+
 private:
 	bool bAreChildrenCreated;
 };
diff --git a/Engine/Plugins/Editor/AssetManagerEditor/Source/AssetManagerEditor/Private/TreeView/SAssetTableTreeView.cpp b/Engine/Plugins/Editor/AssetManagerEditor/Source/AssetManagerEditor/Private/TreeView/SAssetTableTreeView.cpp
--- a/Engine/Plugins/Editor/AssetManagerEditor/Source/AssetManagerEditor/Private/TreeView/SAssetTableTreeView.cpp
+++ b/Engine/Plugins/Editor/AssetManagerEditor/Source/AssetManagerEditor/Private/TreeView/SAssetTableTreeView.cpp
@@ -244,24 +244,218 @@ void SAssetTableTreeView::InitAvailableViewPresets()
 	//////////////////////////////////////////////////
 	// Default View
 
-	class FDefaultViewPreset : public UE::Insights::ITableTreeViewPreset
+	//class FDefaultViewPreset : public UE::Insights::ITableTreeViewPreset
+	//{
+	//public:
+	//	virtual FText GetName() const override
+	//	{
+	//		return LOCTEXT("Default_PresetName", "Default");
+	//	}
+	//	virtual FText GetToolTip() const override
+	//	{
+	//		return LOCTEXT("Default_PresetToolTip", "Default View\nConfigure the tree view to show default asset info.");
+	//	}
+	//	virtual FName GetSortColumn() const override
+	//	{
+	//		return UE::Insights::FTable::GetHierarchyColumnId();
+	//	}
+	//	virtual EColumnSortMode::Type GetSortMode() const override
+	//	{
+	//		return EColumnSortMode::Type::Ascending;
+	//	}
+	//	virtual void SetCurrentGroupings(const TArray<TSharedPtr<UE::Insights::FTreeNodeGrouping>>& InAvailableGroupings, TArray<TSharedPtr<UE::Insights::FTreeNodeGrouping>>& InOutCurrentGroupings) const override
+	//	{
+	//		InOutCurrentGroupings.Reset();
+
+	//		check(InAvailableGroupings[0]->Is<UE::Insights::FTreeNodeGroupingFlat>());
+	//		InOutCurrentGroupings.Add(InAvailableGroupings[0]);
+	//	}
+	//	virtual void GetColumnConfigSet(TArray<UE::Insights::FTableColumnConfig>& InOutConfigSet) const override
+	//	{
+	//		InOutConfigSet.Add({ UE::Insights::FTable::GetHierarchyColumnId(),              true, 400.0f });
+	//		InOutConfigSet.Add({ FAssetTableColumns::CountColumnId,                         true, 100.0f });
+	//		InOutConfigSet.Add({ FAssetTableColumns::TypeColumnId,                         !true, 200.0f });
+	//		InOutConfigSet.Add({ FAssetTableColumns::NameColumnId,                         !true, 200.0f });
+	//		InOutConfigSet.Add({ FAssetTableColumns::PathColumnId,                         !true, 400.0f });
+	//		InOutConfigSet.Add({ FAssetTableColumns::PrimaryTypeColumnId,                   true, 200.0f });
+	//		InOutConfigSet.Add({ FAssetTableColumns::PrimaryNameColumnId,                   true, 200.0f });
+	//		InOutConfigSet.Add({ FAssetTableColumns::StagedCompressedSizeColumnId,          true, 100.0f });
+	//		InOutConfigSet.Add({ FAssetTableColumns::TotalSizeUniqueDependenciesColumnId,  !true, 100.0f });
+	//		InOutConfigSet.Add({ FAssetTableColumns::TotalSizeSharedDependenciesColumnId,  !true, 100.0f });
+	//		InOutConfigSet.Add({ FAssetTableColumns::TotalSizeExternalDependenciesColumnId,!true, 100.0f });
+	//		InOutConfigSet.Add({ FAssetTableColumns::TotalUsageCountColumnId,               true, 100.0f });
+	//		InOutConfigSet.Add({ FAssetTableColumns::ChunksColumnId,                       !true, 200.0f });
+	//		InOutConfigSet.Add({ FAssetTableColumns::NativeClassColumnId,                   true, 200.0f });
+	//		InOutConfigSet.Add({ FAssetTableColumns::PluginNameColumnId,                    true, 200.0f });
+	//	}
+	//};
+	//AvailableViewPresets.Add(MakeShared<FDefaultViewPreset>());
+
+	////////////////////////////////////////////////////
+	//// GameFeaturePlugin, Type, Dependency View
+	//
+	//class FGameFeaturePluginTypeDependencyView : public UE::Insights::ITableTreeViewPreset
+	//{
+	//public:
+	//	virtual FText GetName() const override
+	//	{
+	//		return LOCTEXT("GFPTypeDepView_PresetName", "Dependency Analysis");
+	//	}
+
+	//	virtual FText GetToolTip() const override
+	//	{
+	//		return LOCTEXT("GFPTypeDepView_PresetToolTip", "Dependency Analysis View\nConfigure the tree view to show a breakdown of assets by Game Feature Plugin, Type, and Dependencies.");
+	//	}
+	//	virtual FName GetSortColumn() const override
+	//	{
+	//		return UE::Insights::FTable::GetHierarchyColumnId();
+	//	}
+	//	virtual EColumnSortMode::Type GetSortMode() const override
+	//	{
+	//		return EColumnSortMode::Type::Ascending;
+	//	}
+	//	virtual void SetCurrentGroupings(const TArray<TSharedPtr<UE::Insights::FTreeNodeGrouping>>& InAvailableGroupings, TArray<TSharedPtr<UE::Insights::FTreeNodeGrouping>>& InOutCurrentGroupings) const override
+	//	{
+	//		InOutCurrentGroupings.Reset();
+
+	//		check(InAvailableGroupings[0]->Is<UE::Insights::FTreeNodeGroupingFlat>());
+	//		InOutCurrentGroupings.Add(InAvailableGroupings[0]);
+
+	//		const TSharedPtr<UE::Insights::FTreeNodeGrouping>* GameFeaturePluginGrouping = InAvailableGroupings.FindByPredicate(
+	//			[](TSharedPtr<UE::Insights::FTreeNodeGrouping>& Grouping)
+	//			{
+	//				return Grouping->Is<UE::Insights::FTreeNodeGroupingByUniqueValueCString>() &&
+	//					Grouping->As<UE::Insights::FTreeNodeGroupingByUniqueValueCString>().GetColumnId() == FAssetTableColumns::PluginNameColumnId;
+	//			});
+	//		if (GameFeaturePluginGrouping)
+	//		{
+	//			InOutCurrentGroupings.Add(*GameFeaturePluginGrouping);
+	//		}
+
+	//		const TSharedPtr<UE::Insights::FTreeNodeGrouping>* PrimaryTypeGrouping = InAvailableGroupings.FindByPredicate(
+	//			[](TSharedPtr<UE::Insights::FTreeNodeGrouping>& Grouping)
+	//			{
+	//				return Grouping->Is<UE::Insights::FTreeNodeGroupingByUniqueValueCString>() &&
+	//					Grouping->As<UE::Insights::FTreeNodeGroupingByUniqueValueCString>().GetColumnId() == FAssetTableColumns::TypeColumnId;
+	//			});
+	//		if (PrimaryTypeGrouping)
+	//		{
+	//			InOutCurrentGroupings.Add(*PrimaryTypeGrouping);
+	//		}
+
+	//		const TSharedPtr<UE::Insights::FTreeNodeGrouping>* DependencyGrouping = InAvailableGroupings.FindByPredicate(
+	//			[](TSharedPtr<UE::Insights::FTreeNodeGrouping>& Grouping)
+	//			{
+	//				return Grouping->Is<FAssetDependencyGrouping>();
+	//			});
+	//		if (DependencyGrouping)
+	//		{
+	//			InOutCurrentGroupings.Add(*DependencyGrouping);
+	//		}
+	//	}
+	//	virtual void GetColumnConfigSet(TArray<UE::Insights::FTableColumnConfig>& InOutConfigSet) const override
+	//	{
+	//		InOutConfigSet.Add({ UE::Insights::FTable::GetHierarchyColumnId(),              true, 400.0f });
+	//		InOutConfigSet.Add({ FAssetTableColumns::CountColumnId,                         true, 100.0f });
+	//		InOutConfigSet.Add({ FAssetTableColumns::TypeColumnId,                         !true, 200.0f });
+	//		InOutConfigSet.Add({ FAssetTableColumns::NameColumnId,                         !true, 200.0f });
+	//		InOutConfigSet.Add({ FAssetTableColumns::PathColumnId,                         !true, 400.0f });
+	//		InOutConfigSet.Add({ FAssetTableColumns::PrimaryTypeColumnId,                   true, 200.0f });
+	//		InOutConfigSet.Add({ FAssetTableColumns::PrimaryNameColumnId,                  !true, 200.0f });
+	//		InOutConfigSet.Add({ FAssetTableColumns::StagedCompressedSizeColumnId,          true, 100.0f });
+	//		InOutConfigSet.Add({ FAssetTableColumns::TotalSizeUniqueDependenciesColumnId,   true, 100.0f });
+	//		InOutConfigSet.Add({ FAssetTableColumns::TotalSizeSharedDependenciesColumnId,   true, 100.0f });
+	//		InOutConfigSet.Add({ FAssetTableColumns::TotalSizeExternalDependenciesColumnId,!true, 100.0f });
+	//		InOutConfigSet.Add({ FAssetTableColumns::TotalUsageCountColumnId,              !true, 100.0f });
+	//		InOutConfigSet.Add({ FAssetTableColumns::ChunksColumnId,                       !true, 200.0f });
+	//		InOutConfigSet.Add({ FAssetTableColumns::NativeClassColumnId,                   true, 200.0f });
+	//		InOutConfigSet.Add({ FAssetTableColumns::PluginNameColumnId,                    true, 200.0f });
+	//	}
+	//};
+	//AvailableViewPresets.Add(MakeShared<FGameFeaturePluginTypeDependencyView>());
+
+	////////////////////////////////////////////////////
+	//// Path Breakdown View
+
+	//class FAssetPathViewPreset : public UE::Insights::ITableTreeViewPreset
+	//{
+	//public:
+	//	virtual FText GetName() const override
+	//	{
+	//		return LOCTEXT("Path_PresetName", "Path");
+	//	}
+	//	virtual FText GetToolTip() const override
+	//	{
+	//		return LOCTEXT("Path_PresetToolTip", "Path Breakdown View\nConfigure the tree view to show a breakdown of assets by their path.");
+	//	}
+	//	virtual FName GetSortColumn() const override
+	//	{
+	//		return UE::Insights::FTable::GetHierarchyColumnId();
+	//	}
+	//	virtual EColumnSortMode::Type GetSortMode() const override
+	//	{
+	//		return EColumnSortMode::Type::Ascending;
+	//	}
+	//	virtual void SetCurrentGroupings(const TArray<TSharedPtr<UE::Insights::FTreeNodeGrouping>>& InAvailableGroupings, TArray<TSharedPtr<UE::Insights::FTreeNodeGrouping>>& InOutCurrentGroupings) const override
+	//	{
+	//		InOutCurrentGroupings.Reset();
+
+	//		check(InAvailableGroupings[0]->Is<UE::Insights::FTreeNodeGroupingFlat>());
+	//		InOutCurrentGroupings.Add(InAvailableGroupings[0]);
+
+	//		const TSharedPtr<UE::Insights::FTreeNodeGrouping>* PathGrouping = InAvailableGroupings.FindByPredicate(
+	//			[](TSharedPtr<UE::Insights::FTreeNodeGrouping>& Grouping)
+	//			{
+	//				return Grouping->Is<UE::Insights::FTreeNodeGroupingByPathBreakdown>() &&
+	//					   Grouping->As<UE::Insights::FTreeNodeGroupingByPathBreakdown>().GetColumnId() == FAssetTableColumns::PathColumnId;
+	//			});
+	//		if (PathGrouping)
+	//		{
+	//			InOutCurrentGroupings.Add(*PathGrouping);
+	//		}
+	//	}
+	//	virtual void GetColumnConfigSet(TArray<UE::Insights::FTableColumnConfig>& InOutConfigSet) const override
+	//	{
+	//		InOutConfigSet.Add({ UE::Insights::FTable::GetHierarchyColumnId(),              true, 400.0f });
+	//		InOutConfigSet.Add({ FAssetTableColumns::CountColumnId,                         true, 100.0f });
+	//		InOutConfigSet.Add({ FAssetTableColumns::TypeColumnId,                         !true, 200.0f });
+	//		InOutConfigSet.Add({ FAssetTableColumns::NameColumnId,                         !true, 200.0f });
+	//		InOutConfigSet.Add({ FAssetTableColumns::PathColumnId,                         !true, 400.0f });
+	//		InOutConfigSet.Add({ FAssetTableColumns::PrimaryTypeColumnId,                   true, 200.0f });
+	//		InOutConfigSet.Add({ FAssetTableColumns::PrimaryNameColumnId,                   true, 200.0f });
+	//		InOutConfigSet.Add({ FAssetTableColumns::StagedCompressedSizeColumnId,          true, 100.0f });
+	//		InOutConfigSet.Add({ FAssetTableColumns::TotalSizeUniqueDependenciesColumnId,   true, 100.0f });
+	//		InOutConfigSet.Add({ FAssetTableColumns::TotalSizeSharedDependenciesColumnId,   true, 100.0f });
+	//		InOutConfigSet.Add({ FAssetTableColumns::TotalSizeExternalDependenciesColumnId,!true, 100.0f });
+	//		InOutConfigSet.Add({ FAssetTableColumns::TotalUsageCountColumnId,               true, 100.0f });
+	//		InOutConfigSet.Add({ FAssetTableColumns::ChunksColumnId,                       !true, 200.0f });
+	//		InOutConfigSet.Add({ FAssetTableColumns::NativeClassColumnId,                   true, 200.0f });
+	//		InOutConfigSet.Add({ FAssetTableColumns::PluginNameColumnId,                    true, 200.0f });
+	//	}
+	//};
+	//AvailableViewPresets.Add(MakeShared<FAssetPathViewPreset>());
+
+	////////////////////////////////////////////////////
+	//// Primary Type Breakdown View
+
+	class FPrimaryTypeViewPreset : public UE::Insights::ITableTreeViewPreset
 	{
 	public:
 		virtual FText GetName() const override
 		{
-			return LOCTEXT("Default_PresetName", "Default");
+			return LOCTEXT("PrimaryType_PresetName", "Group By Primary Type");
 		}
 		virtual FText GetToolTip() const override
 		{
-			return LOCTEXT("Default_PresetToolTip", "Default View\nConfigure the tree view to show default asset info.");
+			return LOCTEXT("PrimaryType_PresetToolTip", "Primary Type Breakdown View\nConfigure the tree view to show a breakdown of assets by their primary type.");
 		}
 		virtual FName GetSortColumn() const override
 		{
-			return UE::Insights::FTable::GetHierarchyColumnId();
+			return FAssetTableColumns::StagedCompressedSizeColumnId;
 		}
 		virtual EColumnSortMode::Type GetSortMode() const override
 		{
-			return EColumnSortMode::Type::Ascending;
+			return EColumnSortMode::Type::Descending;
 		}
 		virtual void SetCurrentGroupings(const TArray<TSharedPtr<UE::Insights::FTreeNodeGrouping>>& InAvailableGroupings, TArray<TSharedPtr<UE::Insights::FTreeNodeGrouping>>& InOutCurrentGroupings) const override
 		{
@@ -269,50 +463,61 @@ void SAssetTableTreeView::InitAvailableViewPresets()
 
 			check(InAvailableGroupings[0]->Is<UE::Insights::FTreeNodeGroupingFlat>());
 			InOutCurrentGroupings.Add(InAvailableGroupings[0]);
+
+			const TSharedPtr<UE::Insights::FTreeNodeGrouping>* PrimaryTypeGrouping = InAvailableGroupings.FindByPredicate(
+				[](TSharedPtr<UE::Insights::FTreeNodeGrouping>& Grouping)
+				{
+					return Grouping->Is<UE::Insights::FTreeNodeGroupingByUniqueValueCString>() &&
+						Grouping->As<UE::Insights::FTreeNodeGroupingByUniqueValueCString>().GetColumnId() == FAssetTableColumns::PrimaryTypeColumnId;
+				});
+			if (PrimaryTypeGrouping)
+			{
+				InOutCurrentGroupings.Add(*PrimaryTypeGrouping);
+			}
 		}
 		virtual void GetColumnConfigSet(TArray<UE::Insights::FTableColumnConfig>& InOutConfigSet) const override
 		{
 			InOutConfigSet.Add({ UE::Insights::FTable::GetHierarchyColumnId(),              true, 400.0f });
 			InOutConfigSet.Add({ FAssetTableColumns::CountColumnId,                         true, 100.0f });
-			InOutConfigSet.Add({ FAssetTableColumns::TypeColumnId,                         !true, 200.0f });
-			InOutConfigSet.Add({ FAssetTableColumns::NameColumnId,                         !true, 200.0f });
+			InOutConfigSet.Add({ FAssetTableColumns::StagedCompressedSizeColumnId,          true, 100.0f });
+			InOutConfigSet.Add({ FAssetTableColumns::TypeColumnId,                          true, 200.0f });
+			InOutConfigSet.Add({ FAssetTableColumns::NameColumnId,                          true, 200.0f });
 			InOutConfigSet.Add({ FAssetTableColumns::PathColumnId,                         !true, 400.0f });
 			InOutConfigSet.Add({ FAssetTableColumns::PrimaryTypeColumnId,                   true, 200.0f });
 			InOutConfigSet.Add({ FAssetTableColumns::PrimaryNameColumnId,                   true, 200.0f });
-			InOutConfigSet.Add({ FAssetTableColumns::StagedCompressedSizeColumnId,          true, 100.0f });
-			InOutConfigSet.Add({ FAssetTableColumns::TotalSizeUniqueDependenciesColumnId,  !true, 100.0f });
-			InOutConfigSet.Add({ FAssetTableColumns::TotalSizeSharedDependenciesColumnId,  !true, 100.0f });
-			InOutConfigSet.Add({ FAssetTableColumns::TotalSizeExternalDependenciesColumnId,!true, 100.0f });
+			InOutConfigSet.Add({ FAssetTableColumns::TotalSizeUniqueDependenciesColumnId,   true, 100.0f });
+			InOutConfigSet.Add({ FAssetTableColumns::TotalSizeSharedDependenciesColumnId,   true, 100.0f });
+			InOutConfigSet.Add({ FAssetTableColumns::TotalSizeExternalDependenciesColumnId, true, 100.0f });
 			InOutConfigSet.Add({ FAssetTableColumns::TotalUsageCountColumnId,               true, 100.0f });
 			InOutConfigSet.Add({ FAssetTableColumns::ChunksColumnId,                       !true, 200.0f });
 			InOutConfigSet.Add({ FAssetTableColumns::NativeClassColumnId,                   true, 200.0f });
 			InOutConfigSet.Add({ FAssetTableColumns::PluginNameColumnId,                    true, 200.0f });
 		}
 	};
-	AvailableViewPresets.Add(MakeShared<FDefaultViewPreset>());
+	AvailableViewPresets.Add(MakeShared<FPrimaryTypeViewPreset>());
 
-	//////////////////////////////////////////////////
-	// GameFeaturePlugin, Type, Dependency View
-	
-	class FGameFeaturePluginTypeDependencyView : public UE::Insights::ITableTreeViewPreset
+
+	////////////////////////////////////////////////////
+	//// Class Type Breakdown View
+
+	class FClassTypeViewPreset : public UE::Insights::ITableTreeViewPreset
 	{
 	public:
 		virtual FText GetName() const override
 		{
-			return LOCTEXT("GFPTypeDepView_PresetName", "Dependency Analysis");
+			return LOCTEXT("ClassType_PresetName", "Group By Class Type");
 		}
-
 		virtual FText GetToolTip() const override
 		{
-			return LOCTEXT("GFPTypeDepView_PresetToolTip", "Dependency Analysis View\nConfigure the tree view to show a breakdown of assets by Game Feature Plugin, Type, and Dependencies.");
+			return LOCTEXT("ClassType_PresetToolTip", "Class Type Breakdown View\nConfigure the tree view to show a breakdown of assets by their class type.");
 		}
 		virtual FName GetSortColumn() const override
 		{
-			return UE::Insights::FTable::GetHierarchyColumnId();
+			return FAssetTableColumns::StagedCompressedSizeColumnId;
 		}
 		virtual EColumnSortMode::Type GetSortMode() const override
 		{
-			return EColumnSortMode::Type::Ascending;
+			return EColumnSortMode::Type::Descending;
 		}
 		virtual void SetCurrentGroupings(const TArray<TSharedPtr<UE::Insights::FTreeNodeGrouping>>& InAvailableGroupings, TArray<TSharedPtr<UE::Insights::FTreeNodeGrouping>>& InOutCurrentGroupings) const override
 		{
@@ -321,81 +526,59 @@ void SAssetTableTreeView::InitAvailableViewPresets()
 			check(InAvailableGroupings[0]->Is<UE::Insights::FTreeNodeGroupingFlat>());
 			InOutCurrentGroupings.Add(InAvailableGroupings[0]);
 
-			const TSharedPtr<UE::Insights::FTreeNodeGrouping>* GameFeaturePluginGrouping = InAvailableGroupings.FindByPredicate(
-				[](TSharedPtr<UE::Insights::FTreeNodeGrouping>& Grouping)
-				{
-					return Grouping->Is<UE::Insights::FTreeNodeGroupingByUniqueValueCString>() &&
-						Grouping->As<UE::Insights::FTreeNodeGroupingByUniqueValueCString>().GetColumnId() == FAssetTableColumns::PluginNameColumnId;
-				});
-			if (GameFeaturePluginGrouping)
-			{
-				InOutCurrentGroupings.Add(*GameFeaturePluginGrouping);
-			}
-
 			const TSharedPtr<UE::Insights::FTreeNodeGrouping>* PrimaryTypeGrouping = InAvailableGroupings.FindByPredicate(
 				[](TSharedPtr<UE::Insights::FTreeNodeGrouping>& Grouping)
 				{
 					return Grouping->Is<UE::Insights::FTreeNodeGroupingByUniqueValueCString>() &&
-						Grouping->As<UE::Insights::FTreeNodeGroupingByUniqueValueCString>().GetColumnId() == FAssetTableColumns::TypeColumnId;
+						Grouping->As<UE::Insights::FTreeNodeGroupingByUniqueValueCString>().GetColumnId() == FAssetTableColumns::NativeClassColumnId;
 				});
 			if (PrimaryTypeGrouping)
 			{
 				InOutCurrentGroupings.Add(*PrimaryTypeGrouping);
 			}
-
-			const TSharedPtr<UE::Insights::FTreeNodeGrouping>* DependencyGrouping = InAvailableGroupings.FindByPredicate(
-				[](TSharedPtr<UE::Insights::FTreeNodeGrouping>& Grouping)
-				{
-					return Grouping->Is<FAssetDependencyGrouping>();
-				});
-			if (DependencyGrouping)
-			{
-				InOutCurrentGroupings.Add(*DependencyGrouping);
-			}
 		}
 		virtual void GetColumnConfigSet(TArray<UE::Insights::FTableColumnConfig>& InOutConfigSet) const override
 		{
 			InOutConfigSet.Add({ UE::Insights::FTable::GetHierarchyColumnId(),              true, 400.0f });
 			InOutConfigSet.Add({ FAssetTableColumns::CountColumnId,                         true, 100.0f });
-			InOutConfigSet.Add({ FAssetTableColumns::TypeColumnId,                         !true, 200.0f });
-			InOutConfigSet.Add({ FAssetTableColumns::NameColumnId,                         !true, 200.0f });
+			InOutConfigSet.Add({ FAssetTableColumns::StagedCompressedSizeColumnId,          true, 100.0f });
+			InOutConfigSet.Add({ FAssetTableColumns::TypeColumnId,                          true, 200.0f });
+			InOutConfigSet.Add({ FAssetTableColumns::NameColumnId,                          true, 200.0f });
 			InOutConfigSet.Add({ FAssetTableColumns::PathColumnId,                         !true, 400.0f });
 			InOutConfigSet.Add({ FAssetTableColumns::PrimaryTypeColumnId,                   true, 200.0f });
-			InOutConfigSet.Add({ FAssetTableColumns::PrimaryNameColumnId,                  !true, 200.0f });
-			InOutConfigSet.Add({ FAssetTableColumns::StagedCompressedSizeColumnId,          true, 100.0f });
+			InOutConfigSet.Add({ FAssetTableColumns::PrimaryNameColumnId,                   true, 200.0f });
 			InOutConfigSet.Add({ FAssetTableColumns::TotalSizeUniqueDependenciesColumnId,   true, 100.0f });
 			InOutConfigSet.Add({ FAssetTableColumns::TotalSizeSharedDependenciesColumnId,   true, 100.0f });
-			InOutConfigSet.Add({ FAssetTableColumns::TotalSizeExternalDependenciesColumnId,!true, 100.0f });
-			InOutConfigSet.Add({ FAssetTableColumns::TotalUsageCountColumnId,              !true, 100.0f });
+			InOutConfigSet.Add({ FAssetTableColumns::TotalSizeExternalDependenciesColumnId, true, 100.0f });
+			InOutConfigSet.Add({ FAssetTableColumns::TotalUsageCountColumnId,               true, 100.0f });
 			InOutConfigSet.Add({ FAssetTableColumns::ChunksColumnId,                       !true, 200.0f });
 			InOutConfigSet.Add({ FAssetTableColumns::NativeClassColumnId,                   true, 200.0f });
 			InOutConfigSet.Add({ FAssetTableColumns::PluginNameColumnId,                    true, 200.0f });
 		}
 	};
-	AvailableViewPresets.Add(MakeShared<FGameFeaturePluginTypeDependencyView>());
+	AvailableViewPresets.Add(MakeShared<FClassTypeViewPreset>());
 
-	//////////////////////////////////////////////////
-	// Plugin Dependency View
+	////////////////////////////////////////////////////
+	//// Asset Type Breakdown View
 
-	class FPluginDependencyView : public UE::Insights::ITableTreeViewPreset
+	class FAssetTypeViewPreset : public UE::Insights::ITableTreeViewPreset
 	{
 	public:
 		virtual FText GetName() const override
 		{
-			return LOCTEXT("PluginDepView_PresetName", "Plugin Dependency Analysis");
+			return LOCTEXT("AssetType_PresetName", "Group By Asset Type");
 		}
-
 		virtual FText GetToolTip() const override
 		{
-			return LOCTEXT("PluginDepView_PresetToolTip", "Plugin Dependency Analysis View\nConfigure the tree view to show a breakdown of assets by Game Feature Plugin, showing also the dependencies between plugins.");
+			return LOCTEXT("AssetType_PresetToolTip", "Asset Type Breakdown View\nConfigure the tree view to show a breakdown of assets by their asset type.");
 		}
 		virtual FName GetSortColumn() const override
 		{
-			return UE::Insights::FTable::GetHierarchyColumnId();
+			return FAssetTableColumns::StagedCompressedSizeColumnId;
 		}
 		virtual EColumnSortMode::Type GetSortMode() const override
 		{
-			return EColumnSortMode::Type::Ascending;
+			return EColumnSortMode::Type::Descending;
 		}
 		virtual void SetCurrentGroupings(const TArray<TSharedPtr<UE::Insights::FTreeNodeGrouping>>& InAvailableGroupings, TArray<TSharedPtr<UE::Insights::FTreeNodeGrouping>>& InOutCurrentGroupings) const override
 		{
@@ -404,58 +587,60 @@ void SAssetTableTreeView::InitAvailableViewPresets()
 			check(InAvailableGroupings[0]->Is<UE::Insights::FTreeNodeGroupingFlat>());
 			InOutCurrentGroupings.Add(InAvailableGroupings[0]);
 
-			const TSharedPtr<UE::Insights::FTreeNodeGrouping>* DependencyGrouping = InAvailableGroupings.FindByPredicate(
+			const TSharedPtr<UE::Insights::FTreeNodeGrouping>* PrimaryTypeGrouping = InAvailableGroupings.FindByPredicate(
 				[](TSharedPtr<UE::Insights::FTreeNodeGrouping>& Grouping)
 				{
-					return Grouping->Is<FPluginDependencyGrouping>();
+					return Grouping->Is<UE::Insights::FTreeNodeGroupingByUniqueValueCString>() &&
+						Grouping->As<UE::Insights::FTreeNodeGroupingByUniqueValueCString>().GetColumnId() == FAssetTableColumns::TypeColumnId;
 				});
-			if (DependencyGrouping)
+			if (PrimaryTypeGrouping)
 			{
-				InOutCurrentGroupings.Add(*DependencyGrouping);
+				InOutCurrentGroupings.Add(*PrimaryTypeGrouping);
 			}
 		}
 		virtual void GetColumnConfigSet(TArray<UE::Insights::FTableColumnConfig>& InOutConfigSet) const override
 		{
 			InOutConfigSet.Add({ UE::Insights::FTable::GetHierarchyColumnId(),              true, 400.0f });
 			InOutConfigSet.Add({ FAssetTableColumns::CountColumnId,                         true, 100.0f });
+			InOutConfigSet.Add({ FAssetTableColumns::StagedCompressedSizeColumnId,          true, 100.0f });
 			InOutConfigSet.Add({ FAssetTableColumns::TypeColumnId,                          true, 200.0f });
-			InOutConfigSet.Add({ FAssetTableColumns::NameColumnId,                         !true, 200.0f });
+			InOutConfigSet.Add({ FAssetTableColumns::NameColumnId,                          true, 200.0f });
 			InOutConfigSet.Add({ FAssetTableColumns::PathColumnId,                         !true, 400.0f });
 			InOutConfigSet.Add({ FAssetTableColumns::PrimaryTypeColumnId,                   true, 200.0f });
-			InOutConfigSet.Add({ FAssetTableColumns::PrimaryNameColumnId,                  !true, 200.0f });
-			InOutConfigSet.Add({ FAssetTableColumns::StagedCompressedSizeColumnId,          true, 100.0f });
+			InOutConfigSet.Add({ FAssetTableColumns::PrimaryNameColumnId,                   true, 200.0f });
 			InOutConfigSet.Add({ FAssetTableColumns::TotalSizeUniqueDependenciesColumnId,   true, 100.0f });
 			InOutConfigSet.Add({ FAssetTableColumns::TotalSizeSharedDependenciesColumnId,   true, 100.0f });
-			InOutConfigSet.Add({ FAssetTableColumns::TotalSizeExternalDependenciesColumnId,!true, 100.0f });
-			InOutConfigSet.Add({ FAssetTableColumns::TotalUsageCountColumnId,              !true, 100.0f });
+			InOutConfigSet.Add({ FAssetTableColumns::TotalSizeExternalDependenciesColumnId, true, 100.0f });
+			InOutConfigSet.Add({ FAssetTableColumns::TotalUsageCountColumnId,               true, 100.0f });
 			InOutConfigSet.Add({ FAssetTableColumns::ChunksColumnId,                       !true, 200.0f });
 			InOutConfigSet.Add({ FAssetTableColumns::NativeClassColumnId,                   true, 200.0f });
-			InOutConfigSet.Add({ FAssetTableColumns::PluginNameColumnId,                   !true, 200.0f });
+			InOutConfigSet.Add({ FAssetTableColumns::PluginNameColumnId,                    true, 200.0f });
 		}
 	};
-	AvailableViewPresets.Add(MakeShared<FPluginDependencyView>());
+	AvailableViewPresets.Add(MakeShared<FAssetTypeViewPreset>());
 
 	//////////////////////////////////////////////////
-	// Path Breakdown View
+	// Plugin Group View
 
-	class FAssetPathViewPreset : public UE::Insights::ITableTreeViewPreset
+	class FPluginView : public UE::Insights::ITableTreeViewPreset
 	{
 	public:
 		virtual FText GetName() const override
 		{
-			return LOCTEXT("Path_PresetName", "Path");
+			return LOCTEXT("PluginDepView_PresetName", "Group By Plugin");
 		}
+
 		virtual FText GetToolTip() const override
 		{
-			return LOCTEXT("Path_PresetToolTip", "Path Breakdown View\nConfigure the tree view to show a breakdown of assets by their path.");
+			return LOCTEXT("PluginDepView_PresetToolTip", "Group By Plugin\nGroup assets by plugin.");
 		}
 		virtual FName GetSortColumn() const override
 		{
-			return UE::Insights::FTable::GetHierarchyColumnId();
+			return FAssetTableColumns::StagedCompressedSizeColumnId;
 		}
 		virtual EColumnSortMode::Type GetSortMode() const override
 		{
-			return EColumnSortMode::Type::Ascending;
+			return EColumnSortMode::Type::Descending;
 		}
 		virtual void SetCurrentGroupings(const TArray<TSharedPtr<UE::Insights::FTreeNodeGrouping>>& InAvailableGroupings, TArray<TSharedPtr<UE::Insights::FTreeNodeGrouping>>& InOutCurrentGroupings) const override
 		{
@@ -464,59 +649,71 @@ void SAssetTableTreeView::InitAvailableViewPresets()
 			check(InAvailableGroupings[0]->Is<UE::Insights::FTreeNodeGroupingFlat>());
 			InOutCurrentGroupings.Add(InAvailableGroupings[0]);
 
-			const TSharedPtr<UE::Insights::FTreeNodeGrouping>* PathGrouping = InAvailableGroupings.FindByPredicate(
+			const TSharedPtr<UE::Insights::FTreeNodeGrouping>* PrimaryGrouping = InAvailableGroupings.FindByPredicate(
+				[](TSharedPtr<UE::Insights::FTreeNodeGrouping>& Grouping)
+				{
+					return Grouping->Is<UE::Insights::FTreeNodeGroupingByUniqueValueCString>() &&
+						Grouping->As<UE::Insights::FTreeNodeGroupingByUniqueValueCString>().GetColumnId() == FAssetTableColumns::PluginNameColumnId;
+				});
+			if (PrimaryGrouping)
+			{
+				InOutCurrentGroupings.Add(*PrimaryGrouping);
+			}
+
+			const TSharedPtr<UE::Insights::FTreeNodeGrouping>* SecondaryGrouping = InAvailableGroupings.FindByPredicate(
 				[](TSharedPtr<UE::Insights::FTreeNodeGrouping>& Grouping)
 				{
-					return Grouping->Is<UE::Insights::FTreeNodeGroupingByPathBreakdown>() &&
-						   Grouping->As<UE::Insights::FTreeNodeGroupingByPathBreakdown>().GetColumnId() == FAssetTableColumns::PathColumnId;
+					return Grouping->Is<UE::Insights::FTreeNodeGroupingByUniqueValueCString>() &&
+						Grouping->As<UE::Insights::FTreeNodeGroupingByUniqueValueCString>().GetColumnId() == FAssetTableColumns::TypeColumnId;
 				});
-			if (PathGrouping)
+			if (SecondaryGrouping)
 			{
-				InOutCurrentGroupings.Add(*PathGrouping);
+				InOutCurrentGroupings.Add(*SecondaryGrouping);
 			}
 		}
 		virtual void GetColumnConfigSet(TArray<UE::Insights::FTableColumnConfig>& InOutConfigSet) const override
 		{
-			InOutConfigSet.Add({ UE::Insights::FTable::GetHierarchyColumnId(),              true, 400.0f });
+			InOutConfigSet.Add({ UE::Insights::FTable::GetHierarchyColumnId(),              true, 400.0f });	
 			InOutConfigSet.Add({ FAssetTableColumns::CountColumnId,                         true, 100.0f });
+			InOutConfigSet.Add({ FAssetTableColumns::StagedCompressedSizeColumnId,          true, 100.0f });
 			InOutConfigSet.Add({ FAssetTableColumns::TypeColumnId,                         !true, 200.0f });
 			InOutConfigSet.Add({ FAssetTableColumns::NameColumnId,                         !true, 200.0f });
 			InOutConfigSet.Add({ FAssetTableColumns::PathColumnId,                         !true, 400.0f });
-			InOutConfigSet.Add({ FAssetTableColumns::PrimaryTypeColumnId,                   true, 200.0f });
-			InOutConfigSet.Add({ FAssetTableColumns::PrimaryNameColumnId,                   true, 200.0f });
-			InOutConfigSet.Add({ FAssetTableColumns::StagedCompressedSizeColumnId,          true, 100.0f });
+			InOutConfigSet.Add({ FAssetTableColumns::PrimaryTypeColumnId,                  !true, 200.0f });
+			InOutConfigSet.Add({ FAssetTableColumns::PrimaryNameColumnId,                  !true, 200.0f });
 			InOutConfigSet.Add({ FAssetTableColumns::TotalSizeUniqueDependenciesColumnId,   true, 100.0f });
 			InOutConfigSet.Add({ FAssetTableColumns::TotalSizeSharedDependenciesColumnId,   true, 100.0f });
 			InOutConfigSet.Add({ FAssetTableColumns::TotalSizeExternalDependenciesColumnId,!true, 100.0f });
 			InOutConfigSet.Add({ FAssetTableColumns::TotalUsageCountColumnId,               true, 100.0f });
 			InOutConfigSet.Add({ FAssetTableColumns::ChunksColumnId,                       !true, 200.0f });
 			InOutConfigSet.Add({ FAssetTableColumns::NativeClassColumnId,                   true, 200.0f });
-			InOutConfigSet.Add({ FAssetTableColumns::PluginNameColumnId,                    true, 200.0f });
+			InOutConfigSet.Add({ FAssetTableColumns::PluginNameColumnId,                   !true, 200.0f });
 		}
 	};
-	AvailableViewPresets.Add(MakeShared<FAssetPathViewPreset>());
+	AvailableViewPresets.Add(MakeShared<FPluginView>());
 
 	//////////////////////////////////////////////////
-	// Primary Asset Breakdown View
+	// Plugin Dependency View
 
-	class FPrimaryAssetViewPreset : public UE::Insights::ITableTreeViewPreset
+	class FPluginDependencyView : public UE::Insights::ITableTreeViewPreset
 	{
 	public:
 		virtual FText GetName() const override
 		{
-			return LOCTEXT("PrimaryAsset_PresetName", "Primary Asset");
+			return LOCTEXT("PluginDepView_PresetName", "Plugin Dependency Analysis");
 		}
+
 		virtual FText GetToolTip() const override
 		{
-			return LOCTEXT("PrimaryAsset_PresetToolTip", "Primary Asset Breakdown View\nConfigure the tree view to show a breakdown of assets by their primary asset type/name.");
+			return LOCTEXT("PluginDepView_PresetToolTip", "Plugin Dependency Analysis View\nConfigure the tree view to show a breakdown of assets by Game Feature Plugin, showing also the dependencies between plugins.");
 		}
 		virtual FName GetSortColumn() const override
 		{
-			return UE::Insights::FTable::GetHierarchyColumnId();
+			return FAssetTableColumns::StagedCompressedSizeColumnId;
 		}
 		virtual EColumnSortMode::Type GetSortMode() const override
 		{
-			return EColumnSortMode::Type::Ascending;
+			return EColumnSortMode::Type::Descending;
 		}
 		virtual void SetCurrentGroupings(const TArray<TSharedPtr<UE::Insights::FTreeNodeGrouping>>& InAvailableGroupings, TArray<TSharedPtr<UE::Insights::FTreeNodeGrouping>>& InOutCurrentGroupings) const override
 		{
@@ -525,28 +722,38 @@ void SAssetTableTreeView::InitAvailableViewPresets()
 			check(InAvailableGroupings[0]->Is<UE::Insights::FTreeNodeGroupingFlat>());
 			InOutCurrentGroupings.Add(InAvailableGroupings[0]);
 
-			const TSharedPtr<UE::Insights::FTreeNodeGrouping>* PrimaryTypeGrouping = InAvailableGroupings.FindByPredicate(
+			const TSharedPtr<UE::Insights::FTreeNodeGrouping>* PrimaryGrouping = InAvailableGroupings.FindByPredicate(
+				[](TSharedPtr<UE::Insights::FTreeNodeGrouping>& Grouping)
+				{
+					return Grouping->Is<FPluginDependencyGrouping>();
+				});
+			if (PrimaryGrouping)
+			{
+				InOutCurrentGroupings.Add(*PrimaryGrouping);
+			}
+
+			const TSharedPtr<UE::Insights::FTreeNodeGrouping>* SecondaryGrouping = InAvailableGroupings.FindByPredicate(
 				[](TSharedPtr<UE::Insights::FTreeNodeGrouping>& Grouping)
 				{
 					return Grouping->Is<UE::Insights::FTreeNodeGroupingByUniqueValueCString>() &&
-						   Grouping->As<UE::Insights::FTreeNodeGroupingByUniqueValueCString>().GetColumnId() == FAssetTableColumns::PrimaryTypeColumnId;
+						Grouping->As<UE::Insights::FTreeNodeGroupingByUniqueValueCString>().GetColumnId() == FAssetTableColumns::TypeColumnId;
 				});
-			if (PrimaryTypeGrouping)
+			if (SecondaryGrouping)
 			{
-				InOutCurrentGroupings.Add(*PrimaryTypeGrouping);
+				InOutCurrentGroupings.Add(*SecondaryGrouping);
 			}
 		}
 		virtual void GetColumnConfigSet(TArray<UE::Insights::FTableColumnConfig>& InOutConfigSet) const override
 		{
 			InOutConfigSet.Add({ UE::Insights::FTable::GetHierarchyColumnId(),              true, 400.0f });
 			InOutConfigSet.Add({ FAssetTableColumns::CountColumnId,                         true, 100.0f });
+			InOutConfigSet.Add({ FAssetTableColumns::TotalSizeUniqueDependenciesColumnId,   true, 100.0f });
 			InOutConfigSet.Add({ FAssetTableColumns::TypeColumnId,                         !true, 200.0f });
 			InOutConfigSet.Add({ FAssetTableColumns::NameColumnId,                         !true, 200.0f });
 			InOutConfigSet.Add({ FAssetTableColumns::PathColumnId,                         !true, 400.0f });
 			InOutConfigSet.Add({ FAssetTableColumns::PrimaryTypeColumnId,                  !true, 200.0f });
 			InOutConfigSet.Add({ FAssetTableColumns::PrimaryNameColumnId,                  !true, 200.0f });
 			InOutConfigSet.Add({ FAssetTableColumns::StagedCompressedSizeColumnId,          true, 100.0f });
-			InOutConfigSet.Add({ FAssetTableColumns::TotalSizeUniqueDependenciesColumnId,   true, 100.0f });
 			InOutConfigSet.Add({ FAssetTableColumns::TotalSizeSharedDependenciesColumnId,   true, 100.0f });
 			InOutConfigSet.Add({ FAssetTableColumns::TotalSizeExternalDependenciesColumnId,!true, 100.0f });
 			InOutConfigSet.Add({ FAssetTableColumns::TotalUsageCountColumnId,               true, 100.0f });
@@ -555,7 +762,9 @@ void SAssetTableTreeView::InitAvailableViewPresets()
 			InOutConfigSet.Add({ FAssetTableColumns::PluginNameColumnId,                    true, 200.0f });
 		}
 	};
-	AvailableViewPresets.Add(MakeShared<FPrimaryAssetViewPreset>());
+	AvailableViewPresets.Add(MakeShared<FPluginDependencyView>());
+
+	
 
 	//////////////////////////////////////////////////
 
