Added extra grouping to Asset, Class and Type presets

diff --git a/Engine/Plugins/Editor/AssetManagerEditor/Source/AssetManagerEditor/Private/TreeView/SAssetTableTreeView.cpp b/Engine/Plugins/Editor/AssetManagerEditor/Source/AssetManagerEditor/Private/TreeView/SAssetTableTreeView.cpp
--- a/Engine/Plugins/Editor/AssetManagerEditor/Source/AssetManagerEditor/Private/TreeView/SAssetTableTreeView.cpp
+++ b/Engine/Plugins/Editor/AssetManagerEditor/Source/AssetManagerEditor/Private/TreeView/SAssetTableTreeView.cpp
@@ -241,200 +241,6 @@ END_SLATE_FUNCTION_BUILD_OPTIMIZATION
 
 void SAssetTableTreeView::InitAvailableViewPresets()
 {
-	//////////////////////////////////////////////////
-	// Default View
-
-	//class FDefaultViewPreset : public UE::Insights::ITableTreeViewPreset
-	//{
-	//public:
-	//	virtual FText GetName() const override
-	//	{
-	//		return LOCTEXT("Default_PresetName", "Default");
-	//	}
-	//	virtual FText GetToolTip() const override
-	//	{
-	//		return LOCTEXT("Default_PresetToolTip", "Default View\nConfigure the tree view to show default asset info.");
-	//	}
-	//	virtual FName GetSortColumn() const override
-	//	{
-	//		return UE::Insights::FTable::GetHierarchyColumnId();
-	//	}
-	//	virtual EColumnSortMode::Type GetSortMode() const override
-	//	{
-	//		return EColumnSortMode::Type::Ascending;
-	//	}
-	//	virtual void SetCurrentGroupings(const TArray<TSharedPtr<UE::Insights::FTreeNodeGrouping>>& InAvailableGroupings, TArray<TSharedPtr<UE::Insights::FTreeNodeGrouping>>& InOutCurrentGroupings) const override
-	//	{
-	//		InOutCurrentGroupings.Reset();
-
-	//		check(InAvailableGroupings[0]->Is<UE::Insights::FTreeNodeGroupingFlat>());
-	//		InOutCurrentGroupings.Add(InAvailableGroupings[0]);
-	//	}
-	//	virtual void GetColumnConfigSet(TArray<UE::Insights::FTableColumnConfig>& InOutConfigSet) const override
-	//	{
-	//		InOutConfigSet.Add({ UE::Insights::FTable::GetHierarchyColumnId(),              true, 400.0f });
-	//		InOutConfigSet.Add({ FAssetTableColumns::CountColumnId,                         true, 100.0f });
-	//		InOutConfigSet.Add({ FAssetTableColumns::TypeColumnId,                         !true, 200.0f });
-	//		InOutConfigSet.Add({ FAssetTableColumns::NameColumnId,                         !true, 200.0f });
-	//		InOutConfigSet.Add({ FAssetTableColumns::PathColumnId,                         !true, 400.0f });
-	//		InOutConfigSet.Add({ FAssetTableColumns::PrimaryTypeColumnId,                   true, 200.0f });
-	//		InOutConfigSet.Add({ FAssetTableColumns::PrimaryNameColumnId,                   true, 200.0f });
-	//		InOutConfigSet.Add({ FAssetTableColumns::StagedCompressedSizeColumnId,          true, 100.0f });
-	//		InOutConfigSet.Add({ FAssetTableColumns::TotalSizeUniqueDependenciesColumnId,  !true, 100.0f });
-	//		InOutConfigSet.Add({ FAssetTableColumns::TotalSizeSharedDependenciesColumnId,  !true, 100.0f });
-	//		InOutConfigSet.Add({ FAssetTableColumns::TotalSizeExternalDependenciesColumnId,!true, 100.0f });
-	//		InOutConfigSet.Add({ FAssetTableColumns::TotalUsageCountColumnId,               true, 100.0f });
-	//		InOutConfigSet.Add({ FAssetTableColumns::ChunksColumnId,                       !true, 200.0f });
-	//		InOutConfigSet.Add({ FAssetTableColumns::NativeClassColumnId,                   true, 200.0f });
-	//		InOutConfigSet.Add({ FAssetTableColumns::PluginNameColumnId,                    true, 200.0f });
-	//	}
-	//};
-	//AvailableViewPresets.Add(MakeShared<FDefaultViewPreset>());
-
-	////////////////////////////////////////////////////
-	//// GameFeaturePlugin, Type, Dependency View
-	//
-	//class FGameFeaturePluginTypeDependencyView : public UE::Insights::ITableTreeViewPreset
-	//{
-	//public:
-	//	virtual FText GetName() const override
-	//	{
-	//		return LOCTEXT("GFPTypeDepView_PresetName", "Dependency Analysis");
-	//	}
-
-	//	virtual FText GetToolTip() const override
-	//	{
-	//		return LOCTEXT("GFPTypeDepView_PresetToolTip", "Dependency Analysis View\nConfigure the tree view to show a breakdown of assets by Game Feature Plugin, Type, and Dependencies.");
-	//	}
-	//	virtual FName GetSortColumn() const override
-	//	{
-	//		return UE::Insights::FTable::GetHierarchyColumnId();
-	//	}
-	//	virtual EColumnSortMode::Type GetSortMode() const override
-	//	{
-	//		return EColumnSortMode::Type::Ascending;
-	//	}
-	//	virtual void SetCurrentGroupings(const TArray<TSharedPtr<UE::Insights::FTreeNodeGrouping>>& InAvailableGroupings, TArray<TSharedPtr<UE::Insights::FTreeNodeGrouping>>& InOutCurrentGroupings) const override
-	//	{
-	//		InOutCurrentGroupings.Reset();
-
-	//		check(InAvailableGroupings[0]->Is<UE::Insights::FTreeNodeGroupingFlat>());
-	//		InOutCurrentGroupings.Add(InAvailableGroupings[0]);
-
-	//		const TSharedPtr<UE::Insights::FTreeNodeGrouping>* GameFeaturePluginGrouping = InAvailableGroupings.FindByPredicate(
-	//			[](TSharedPtr<UE::Insights::FTreeNodeGrouping>& Grouping)
-	//			{
-	//				return Grouping->Is<UE::Insights::FTreeNodeGroupingByUniqueValueCString>() &&
-	//					Grouping->As<UE::Insights::FTreeNodeGroupingByUniqueValueCString>().GetColumnId() == FAssetTableColumns::PluginNameColumnId;
-	//			});
-	//		if (GameFeaturePluginGrouping)
-	//		{
-	//			InOutCurrentGroupings.Add(*GameFeaturePluginGrouping);
-	//		}
-
-	//		const TSharedPtr<UE::Insights::FTreeNodeGrouping>* PrimaryTypeGrouping = InAvailableGroupings.FindByPredicate(
-	//			[](TSharedPtr<UE::Insights::FTreeNodeGrouping>& Grouping)
-	//			{
-	//				return Grouping->Is<UE::Insights::FTreeNodeGroupingByUniqueValueCString>() &&
-	//					Grouping->As<UE::Insights::FTreeNodeGroupingByUniqueValueCString>().GetColumnId() == FAssetTableColumns::TypeColumnId;
-	//			});
-	//		if (PrimaryTypeGrouping)
-	//		{
-	//			InOutCurrentGroupings.Add(*PrimaryTypeGrouping);
-	//		}
-
-	//		const TSharedPtr<UE::Insights::FTreeNodeGrouping>* DependencyGrouping = InAvailableGroupings.FindByPredicate(
-	//			[](TSharedPtr<UE::Insights::FTreeNodeGrouping>& Grouping)
-	//			{
-	//				return Grouping->Is<FAssetDependencyGrouping>();
-	//			});
-	//		if (DependencyGrouping)
-	//		{
-	//			InOutCurrentGroupings.Add(*DependencyGrouping);
-	//		}
-	//	}
-	//	virtual void GetColumnConfigSet(TArray<UE::Insights::FTableColumnConfig>& InOutConfigSet) const override
-	//	{
-	//		InOutConfigSet.Add({ UE::Insights::FTable::GetHierarchyColumnId(),              true, 400.0f });
-	//		InOutConfigSet.Add({ FAssetTableColumns::CountColumnId,                         true, 100.0f });
-	//		InOutConfigSet.Add({ FAssetTableColumns::TypeColumnId,                         !true, 200.0f });
-	//		InOutConfigSet.Add({ FAssetTableColumns::NameColumnId,                         !true, 200.0f });
-	//		InOutConfigSet.Add({ FAssetTableColumns::PathColumnId,                         !true, 400.0f });
-	//		InOutConfigSet.Add({ FAssetTableColumns::PrimaryTypeColumnId,                   true, 200.0f });
-	//		InOutConfigSet.Add({ FAssetTableColumns::PrimaryNameColumnId,                  !true, 200.0f });
-	//		InOutConfigSet.Add({ FAssetTableColumns::StagedCompressedSizeColumnId,          true, 100.0f });
-	//		InOutConfigSet.Add({ FAssetTableColumns::TotalSizeUniqueDependenciesColumnId,   true, 100.0f });
-	//		InOutConfigSet.Add({ FAssetTableColumns::TotalSizeSharedDependenciesColumnId,   true, 100.0f });
-	//		InOutConfigSet.Add({ FAssetTableColumns::TotalSizeExternalDependenciesColumnId,!true, 100.0f });
-	//		InOutConfigSet.Add({ FAssetTableColumns::TotalUsageCountColumnId,              !true, 100.0f });
-	//		InOutConfigSet.Add({ FAssetTableColumns::ChunksColumnId,                       !true, 200.0f });
-	//		InOutConfigSet.Add({ FAssetTableColumns::NativeClassColumnId,                   true, 200.0f });
-	//		InOutConfigSet.Add({ FAssetTableColumns::PluginNameColumnId,                    true, 200.0f });
-	//	}
-	//};
-	//AvailableViewPresets.Add(MakeShared<FGameFeaturePluginTypeDependencyView>());
-
-	////////////////////////////////////////////////////
-	//// Path Breakdown View
-
-	//class FAssetPathViewPreset : public UE::Insights::ITableTreeViewPreset
-	//{
-	//public:
-	//	virtual FText GetName() const override
-	//	{
-	//		return LOCTEXT("Path_PresetName", "Path");
-	//	}
-	//	virtual FText GetToolTip() const override
-	//	{
-	//		return LOCTEXT("Path_PresetToolTip", "Path Breakdown View\nConfigure the tree view to show a breakdown of assets by their path.");
-	//	}
-	//	virtual FName GetSortColumn() const override
-	//	{
-	//		return UE::Insights::FTable::GetHierarchyColumnId();
-	//	}
-	//	virtual EColumnSortMode::Type GetSortMode() const override
-	//	{
-	//		return EColumnSortMode::Type::Ascending;
-	//	}
-	//	virtual void SetCurrentGroupings(const TArray<TSharedPtr<UE::Insights::FTreeNodeGrouping>>& InAvailableGroupings, TArray<TSharedPtr<UE::Insights::FTreeNodeGrouping>>& InOutCurrentGroupings) const override
-	//	{
-	//		InOutCurrentGroupings.Reset();
-
-	//		check(InAvailableGroupings[0]->Is<UE::Insights::FTreeNodeGroupingFlat>());
-	//		InOutCurrentGroupings.Add(InAvailableGroupings[0]);
-
-	//		const TSharedPtr<UE::Insights::FTreeNodeGrouping>* PathGrouping = InAvailableGroupings.FindByPredicate(
-	//			[](TSharedPtr<UE::Insights::FTreeNodeGrouping>& Grouping)
-	//			{
-	//				return Grouping->Is<UE::Insights::FTreeNodeGroupingByPathBreakdown>() &&
-	//					   Grouping->As<UE::Insights::FTreeNodeGroupingByPathBreakdown>().GetColumnId() == FAssetTableColumns::PathColumnId;
-	//			});
-	//		if (PathGrouping)
-	//		{
-	//			InOutCurrentGroupings.Add(*PathGrouping);
-	//		}
-	//	}
-	//	virtual void GetColumnConfigSet(TArray<UE::Insights::FTableColumnConfig>& InOutConfigSet) const override
-	//	{
-	//		InOutConfigSet.Add({ UE::Insights::FTable::GetHierarchyColumnId(),              true, 400.0f });
-	//		InOutConfigSet.Add({ FAssetTableColumns::CountColumnId,                         true, 100.0f });
-	//		InOutConfigSet.Add({ FAssetTableColumns::TypeColumnId,                         !true, 200.0f });
-	//		InOutConfigSet.Add({ FAssetTableColumns::NameColumnId,                         !true, 200.0f });
-	//		InOutConfigSet.Add({ FAssetTableColumns::PathColumnId,                         !true, 400.0f });
-	//		InOutConfigSet.Add({ FAssetTableColumns::PrimaryTypeColumnId,                   true, 200.0f });
-	//		InOutConfigSet.Add({ FAssetTableColumns::PrimaryNameColumnId,                   true, 200.0f });
-	//		InOutConfigSet.Add({ FAssetTableColumns::StagedCompressedSizeColumnId,          true, 100.0f });
-	//		InOutConfigSet.Add({ FAssetTableColumns::TotalSizeUniqueDependenciesColumnId,   true, 100.0f });
-	//		InOutConfigSet.Add({ FAssetTableColumns::TotalSizeSharedDependenciesColumnId,   true, 100.0f });
-	//		InOutConfigSet.Add({ FAssetTableColumns::TotalSizeExternalDependenciesColumnId,!true, 100.0f });
-	//		InOutConfigSet.Add({ FAssetTableColumns::TotalUsageCountColumnId,               true, 100.0f });
-	//		InOutConfigSet.Add({ FAssetTableColumns::ChunksColumnId,                       !true, 200.0f });
-	//		InOutConfigSet.Add({ FAssetTableColumns::NativeClassColumnId,                   true, 200.0f });
-	//		InOutConfigSet.Add({ FAssetTableColumns::PluginNameColumnId,                    true, 200.0f });
-	//	}
-	//};
-	//AvailableViewPresets.Add(MakeShared<FAssetPathViewPreset>());
-
 	////////////////////////////////////////////////////
 	//// Primary Type Breakdown View
 
@@ -474,6 +280,17 @@ void SAssetTableTreeView::InitAvailableViewPresets()
 			{
 				InOutCurrentGroupings.Add(*PrimaryTypeGrouping);
 			}
+
+			const TSharedPtr<UE::Insights::FTreeNodeGrouping>* SecondaryTypeGrouping = InAvailableGroupings.FindByPredicate(
+				[](TSharedPtr<UE::Insights::FTreeNodeGrouping>& Grouping)
+				{
+					return Grouping->Is<UE::Insights::FTreeNodeGroupingByUniqueValueCString>() &&
+						Grouping->As<UE::Insights::FTreeNodeGroupingByUniqueValueCString>().GetColumnId() == FAssetTableColumns::PluginNameColumnId;
+				});
+			if (SecondaryTypeGrouping)
+			{
+				InOutCurrentGroupings.Add(*SecondaryTypeGrouping);
+			}
 		}
 		virtual void GetColumnConfigSet(TArray<UE::Insights::FTableColumnConfig>& InOutConfigSet) const override
 		{
@@ -536,6 +353,17 @@ void SAssetTableTreeView::InitAvailableViewPresets()
 			{
 				InOutCurrentGroupings.Add(*PrimaryTypeGrouping);
 			}
+
+			const TSharedPtr<UE::Insights::FTreeNodeGrouping>* SecondaryTypeGrouping = InAvailableGroupings.FindByPredicate(
+				[](TSharedPtr<UE::Insights::FTreeNodeGrouping>& Grouping)
+				{
+					return Grouping->Is<UE::Insights::FTreeNodeGroupingByUniqueValueCString>() &&
+						Grouping->As<UE::Insights::FTreeNodeGroupingByUniqueValueCString>().GetColumnId() == FAssetTableColumns::PluginNameColumnId;
+				});
+			if (SecondaryTypeGrouping)
+			{
+				InOutCurrentGroupings.Add(*SecondaryTypeGrouping);
+			}
 		}
 		virtual void GetColumnConfigSet(TArray<UE::Insights::FTableColumnConfig>& InOutConfigSet) const override
 		{
@@ -597,6 +425,17 @@ void SAssetTableTreeView::InitAvailableViewPresets()
 			{
 				InOutCurrentGroupings.Add(*PrimaryTypeGrouping);
 			}
+
+			const TSharedPtr<UE::Insights::FTreeNodeGrouping>* SecondaryTypeGrouping = InAvailableGroupings.FindByPredicate(
+				[](TSharedPtr<UE::Insights::FTreeNodeGrouping>& Grouping)
+				{
+					return Grouping->Is<UE::Insights::FTreeNodeGroupingByUniqueValueCString>() &&
+						Grouping->As<UE::Insights::FTreeNodeGroupingByUniqueValueCString>().GetColumnId() == FAssetTableColumns::PluginNameColumnId;
+				});
+			if (SecondaryTypeGrouping)
+			{
+				InOutCurrentGroupings.Add(*SecondaryTypeGrouping);
+			}
 		}
 		virtual void GetColumnConfigSet(TArray<UE::Insights::FTableColumnConfig>& InOutConfigSet) const override
 		{
