Moved columns about. Put Size further left

diff --git a/Engine/Plugins/Editor/AssetManagerEditor/Source/AssetManagerEditor/Private/TreeView/SAssetTableTreeView.cpp b/Engine/Plugins/Editor/AssetManagerEditor/Source/AssetManagerEditor/Private/TreeView/SAssetTableTreeView.cpp
--- a/Engine/Plugins/Editor/AssetManagerEditor/Source/AssetManagerEditor/Private/TreeView/SAssetTableTreeView.cpp
+++ b/Engine/Plugins/Editor/AssetManagerEditor/Source/AssetManagerEditor/Private/TreeView/SAssetTableTreeView.cpp
@@ -512,21 +512,21 @@ void SAssetTableTreeView::InitAvailableViewPresets()
 		}
 		virtual void GetColumnConfigSet(TArray<UE::Insights::FTableColumnConfig>& InOutConfigSet) const override
 		{
-			InOutConfigSet.Add({ UE::Insights::FTable::GetHierarchyColumnId(),              true, 400.0f });	
+			InOutConfigSet.Add({ UE::Insights::FTable::GetHierarchyColumnId(),              true, 400.0f });
 			InOutConfigSet.Add({ FAssetTableColumns::CountColumnId,                         true, 100.0f });
 			InOutConfigSet.Add({ FAssetTableColumns::StagedCompressedSizeColumnId,          true, 100.0f });
-			InOutConfigSet.Add({ FAssetTableColumns::TypeColumnId,                         !true, 200.0f });
-			InOutConfigSet.Add({ FAssetTableColumns::NameColumnId,                         !true, 200.0f });
+			InOutConfigSet.Add({ FAssetTableColumns::TypeColumnId,                          true, 200.0f });
+			InOutConfigSet.Add({ FAssetTableColumns::NameColumnId,                          true, 200.0f });
 			InOutConfigSet.Add({ FAssetTableColumns::PathColumnId,                         !true, 400.0f });
-			InOutConfigSet.Add({ FAssetTableColumns::PrimaryTypeColumnId,                  !true, 200.0f });
-			InOutConfigSet.Add({ FAssetTableColumns::PrimaryNameColumnId,                  !true, 200.0f });
+			InOutConfigSet.Add({ FAssetTableColumns::PrimaryTypeColumnId,                   true, 200.0f });
+			InOutConfigSet.Add({ FAssetTableColumns::PrimaryNameColumnId,                   true, 200.0f });
 			InOutConfigSet.Add({ FAssetTableColumns::TotalSizeUniqueDependenciesColumnId,   true, 100.0f });
 			InOutConfigSet.Add({ FAssetTableColumns::TotalSizeSharedDependenciesColumnId,   true, 100.0f });
-			InOutConfigSet.Add({ FAssetTableColumns::TotalSizeExternalDependenciesColumnId,!true, 100.0f });
+			InOutConfigSet.Add({ FAssetTableColumns::TotalSizeExternalDependenciesColumnId, true, 100.0f });
 			InOutConfigSet.Add({ FAssetTableColumns::TotalUsageCountColumnId,               true, 100.0f });
 			InOutConfigSet.Add({ FAssetTableColumns::ChunksColumnId,                       !true, 200.0f });
 			InOutConfigSet.Add({ FAssetTableColumns::NativeClassColumnId,                   true, 200.0f });
-			InOutConfigSet.Add({ FAssetTableColumns::PluginNameColumnId,                   !true, 200.0f });
+			InOutConfigSet.Add({ FAssetTableColumns::PluginNameColumnId,                    true, 200.0f });
 		}
 	};
 	AvailableViewPresets.Add(MakeShared<FPluginView>());
@@ -586,15 +586,15 @@ void SAssetTableTreeView::InitAvailableViewPresets()
 		{
 			InOutConfigSet.Add({ UE::Insights::FTable::GetHierarchyColumnId(),              true, 400.0f });
 			InOutConfigSet.Add({ FAssetTableColumns::CountColumnId,                         true, 100.0f });
-			InOutConfigSet.Add({ FAssetTableColumns::TotalSizeUniqueDependenciesColumnId,   true, 100.0f });
-			InOutConfigSet.Add({ FAssetTableColumns::TypeColumnId,                         !true, 200.0f });
-			InOutConfigSet.Add({ FAssetTableColumns::NameColumnId,                         !true, 200.0f });
-			InOutConfigSet.Add({ FAssetTableColumns::PathColumnId,                         !true, 400.0f });
-			InOutConfigSet.Add({ FAssetTableColumns::PrimaryTypeColumnId,                  !true, 200.0f });
-			InOutConfigSet.Add({ FAssetTableColumns::PrimaryNameColumnId,                  !true, 200.0f });
 			InOutConfigSet.Add({ FAssetTableColumns::StagedCompressedSizeColumnId,          true, 100.0f });
+			InOutConfigSet.Add({ FAssetTableColumns::TypeColumnId,                          true, 200.0f });
+			InOutConfigSet.Add({ FAssetTableColumns::NameColumnId,                          true, 200.0f });
+			InOutConfigSet.Add({ FAssetTableColumns::PathColumnId,                         !true, 400.0f });
+			InOutConfigSet.Add({ FAssetTableColumns::PrimaryTypeColumnId,                   true, 200.0f });
+			InOutConfigSet.Add({ FAssetTableColumns::PrimaryNameColumnId,                   true, 200.0f });
+			InOutConfigSet.Add({ FAssetTableColumns::TotalSizeUniqueDependenciesColumnId,   true, 100.0f });
 			InOutConfigSet.Add({ FAssetTableColumns::TotalSizeSharedDependenciesColumnId,   true, 100.0f });
-			InOutConfigSet.Add({ FAssetTableColumns::TotalSizeExternalDependenciesColumnId,!true, 100.0f });
+			InOutConfigSet.Add({ FAssetTableColumns::TotalSizeExternalDependenciesColumnId, true, 100.0f });
 			InOutConfigSet.Add({ FAssetTableColumns::TotalUsageCountColumnId,               true, 100.0f });
 			InOutConfigSet.Add({ FAssetTableColumns::ChunksColumnId,                       !true, 200.0f });
 			InOutConfigSet.Add({ FAssetTableColumns::NativeClassColumnId,                   true, 200.0f });
