Fixing a loctext name conflict in Asset Disk Size Tool window titles.

diff --git a/Engine/Plugins/Editor/AssetManagerEditor/Source/AssetManagerEditor/Private/AssetManagerEditorModule.cpp b/Engine/Plugins/Editor/AssetManagerEditor/Source/AssetManagerEditor/Private/AssetManagerEditorModule.cpp
--- a/Engine/Plugins/Editor/AssetManagerEditor/Source/AssetManagerEditor/Private/AssetManagerEditorModule.cpp
+++ b/Engine/Plugins/Editor/AssetManagerEditor/Source/AssetManagerEditor/Private/AssetManagerEditorModule.cpp
@@ -520,15 +520,15 @@ void FAssetManagerEditorModule::StartupModule()
 		FGlobalTabmanager::Get()->RegisterDefaultTabWindowSize(AssetAuditTabName, FVector2D(1080, 600));
 
 		FGlobalTabmanager::Get()->RegisterNomadTabSpawner(AssetDiskSizeTabName, FOnSpawnTab::CreateRaw(this, &FAssetManagerEditorModule::SpawnAssetDiskSizeTab))
-			.SetDisplayName(LOCTEXT("AssetDiskSizeTitle", "Asset Disk Size 1"))
-			.SetTooltipText(LOCTEXT("AssetDiskSizeTooltip", "Open Asset Disk Size window, allows to analyze disk size for assets."))
+			.SetDisplayName(LOCTEXT("AssetDiskSize1Title", "Asset Disk Size 1"))
+			.SetTooltipText(LOCTEXT("AssetDiskSize1Tooltip", "Open Asset Disk Size window, allows to analyze disk size for assets."))
 			.SetGroup(WorkspaceMenu::GetMenuStructure().GetDeveloperToolsAuditCategory())
 			.SetIcon(FSlateIcon(FAppStyle::GetAppStyleSetName(), "LevelEditor.Audit"));
 		FGlobalTabmanager::Get()->RegisterDefaultTabWindowSize(AssetDiskSizeTabName, FVector2D(1080, 600));
 
 		FGlobalTabmanager::Get()->RegisterNomadTabSpawner(AssetDiskSize2TabName, FOnSpawnTab::CreateRaw(this, &FAssetManagerEditorModule::SpawnAssetDiskSizeTab))
-			.SetDisplayName(LOCTEXT("AssetDiskSizeTitle", "Asset Disk Size 2"))
-			.SetTooltipText(LOCTEXT("AssetDiskSizeTooltip", "Open Asset Disk Size window, allows to analyze disk size for assets."))
+			.SetDisplayName(LOCTEXT("AssetDiskSize2Title", "Asset Disk Size 2"))
+			.SetTooltipText(LOCTEXT("AssetDiskSize2Tooltip", "Open Asset Disk Size window, allows to analyze disk size for assets."))
 			.SetGroup(WorkspaceMenu::GetMenuStructure().GetDeveloperToolsAuditCategory())
 			.SetIcon(FSlateIcon(FAppStyle::GetAppStyleSetName(), "LevelEditor.Audit"));
 		FGlobalTabmanager::Get()->RegisterDefaultTabWindowSize(AssetDiskSize2TabName, FVector2D(1080, 600));
