Stop aggregating self size from plugin dependencies to the parent plugin node. Minor naming cleanup (from earlier copy pasting of code).

diff --git a/Engine/Plugins/Editor/AssetManagerEditor/Source/AssetManagerEditor/Private/TreeView/AssetTable.cpp b/Engine/Plugins/Editor/AssetManagerEditor/Source/AssetManagerEditor/Private/TreeView/AssetTable.cpp
--- a/Engine/Plugins/Editor/AssetManagerEditor/Source/AssetManagerEditor/Private/TreeView/AssetTable.cpp
+++ b/Engine/Plugins/Editor/AssetManagerEditor/Source/AssetManagerEditor/Private/TreeView/AssetTable.cpp
@@ -621,7 +621,6 @@ void FAssetTable::AddDefaultColumns()
 
 		AddColumn(ColumnRef);
 	}
-
 	//////////////////////////////////////////////////
 	// Staged Compressed Size Column
 	{
@@ -656,7 +655,7 @@ void FAssetTable::AddDefaultColumns()
 					const FAssetTableRow& Asset = TreeNode.GetAssetChecked();
 					return FTableCellValue(static_cast<int64>(Asset.GetStagedCompressedSize()));
 				}
-				else if (Node.IsGroup())
+				else if (Node.IsGroup() && !Node.Is<FPluginDependenciesGroupNode>())
 				{
 					const FTableTreeNode& NodePtr = static_cast<const FTableTreeNode&>(Node);
 					if (NodePtr.HasAggregatedValue(Column.GetId()))
@@ -727,6 +726,7 @@ void FAssetTable::AddDefaultColumns()
 
 		AddColumn(ColumnRef);
 	}
+
 	//////////////////////////////////////////////////
 	// Name Column
 	{
@@ -895,7 +895,7 @@ void FAssetTable::AddDefaultColumns()
 
 		AddColumn(ColumnRef);
 	}
-	
+
 	//////////////////////////////////////////////////
 	// Total Size of Unique Dependencies
 	{
@@ -981,7 +981,7 @@ void FAssetTable::AddDefaultColumns()
 
 		Column.SetDataType(ETableCellDataType::Int64);
 
-		class FTotalSizeUniqueDependenciesValueGetter : public FTableCellValueGetter
+		class FTotalSizeSharedDependenciesValueGetter : public FTableCellValueGetter
 		{
 		public:
 			virtual const TOptional<FTableCellValue> GetValue(const FTableColumn& Column, const FBaseTreeNode& Node) const override
@@ -1011,7 +1011,7 @@ void FAssetTable::AddDefaultColumns()
 				return TOptional<FTableCellValue>();
 			}
 		};
-		TSharedRef<ITableCellValueGetter> Getter = MakeShared<FTotalSizeUniqueDependenciesValueGetter>();
+		TSharedRef<ITableCellValueGetter> Getter = MakeShared<FTotalSizeSharedDependenciesValueGetter>();
 		Column.SetValueGetter(Getter);
 
 		TSharedRef<ITableCellValueFormatter> Formatter = MakeShared<FInt64ValueFormatterAsMemory>();
@@ -1047,7 +1047,7 @@ void FAssetTable::AddDefaultColumns()
 
 		Column.SetDataType(ETableCellDataType::Int64);
 
-		class FTotalSizeUniqueDependenciesValueGetter : public FTableCellValueGetter
+		class FTotalSizeExternalDependenciesValueGetter : public FTableCellValueGetter
 		{
 		public:
 			virtual const TOptional<FTableCellValue> GetValue(const FTableColumn& Column, const FBaseTreeNode& Node) const override
@@ -1077,7 +1077,7 @@ void FAssetTable::AddDefaultColumns()
 				return TOptional<FTableCellValue>();
 			}
 		};
-		TSharedRef<ITableCellValueGetter> Getter = MakeShared<FTotalSizeUniqueDependenciesValueGetter>();
+		TSharedRef<ITableCellValueGetter> Getter = MakeShared<FTotalSizeExternalDependenciesValueGetter>();
 		Column.SetValueGetter(Getter);
 
 		TSharedRef<ITableCellValueFormatter> Formatter = MakeShared<FInt64ValueFormatterAsMemory>();
