UX improvements for Asset Disk Size tooling
Assets types now use the asset icon
Fixed selected/default grouping not being applied when first opening an asset registry

diff --git a/Engine/Plugins/Editor/AssetManagerEditor/Source/AssetManagerEditor/Private/TreeView/AssetTreeNode.cpp b/Engine/Plugins/Editor/AssetManagerEditor/Source/AssetManagerEditor/Private/TreeView/AssetTreeNode.cpp
--- a/Engine/Plugins/Editor/AssetManagerEditor/Source/AssetManagerEditor/Private/TreeView/AssetTreeNode.cpp
+++ b/Engine/Plugins/Editor/AssetManagerEditor/Source/AssetManagerEditor/Private/TreeView/AssetTreeNode.cpp
@@ -26,11 +26,16 @@ const FSlateBrush* FAssetTreeNode::GetIcon(EStyle Style) const
 	{
 		default:
 		case EStyle::EDefault:
-		case EStyle::EAsset:
 		{
 			return UE::Insights::FInsightsStyle::GetBrush("Icons.Leaf.TreeItem");
 			break;
 		}
+
+		case EStyle::EAsset:
+		{
+			return UE::Insights::FInsightsStyle::GetBrush("Icons.Asset.TreeItem");
+			break;
+		}
 				
 		case EStyle::EGroup:
 		{
diff --git a/Engine/Plugins/Editor/AssetManagerEditor/Source/AssetManagerEditor/Private/TreeView/SAssetTableTreeView.cpp b/Engine/Plugins/Editor/AssetManagerEditor/Source/AssetManagerEditor/Private/TreeView/SAssetTableTreeView.cpp
--- a/Engine/Plugins/Editor/AssetManagerEditor/Source/AssetManagerEditor/Private/TreeView/SAssetTableTreeView.cpp
+++ b/Engine/Plugins/Editor/AssetManagerEditor/Source/AssetManagerEditor/Private/TreeView/SAssetTableTreeView.cpp
@@ -107,6 +107,7 @@
 		{
 			bNeedsToRebuild = false;
 			RebuildTree(true);
+			ApplyViewPreset(*SelectedViewPreset);
 		}
 	}
 }
@@ -242,18 +243,18 @@
 void SAssetTableTreeView::InitAvailableViewPresets()
 {
 	////////////////////////////////////////////////////
-	//// Primary Type Breakdown View
+	//// Asset Type Breakdown View
 
-	class FPrimaryTypeViewPreset : public UE::Insights::ITableTreeViewPreset
+	class FAssetTypeViewPreset : public UE::Insights::ITableTreeViewPreset
 	{
 	public:
 		virtual FText GetName() const override
 		{
-			return LOCTEXT("PrimaryType_PresetName", "Group By Primary Type");
+			return LOCTEXT("AssetType_PresetName", "Group By Asset Type");
 		}
 		virtual FText GetToolTip() const override
 		{
-			return LOCTEXT("PrimaryType_PresetToolTip", "Primary Type Breakdown View\nConfigure the tree view to show a breakdown of assets by their primary type.");
+			return LOCTEXT("AssetType_PresetToolTip", "Asset Type Breakdown View\nConfigure the tree view to show a breakdown of assets by their asset type.");
 		}
 		virtual FName GetSortColumn() const override
 		{
@@ -274,7 +275,7 @@
 				[](TSharedPtr<UE::Insights::FTreeNodeGrouping>& Grouping)
 				{
 					return Grouping->Is<UE::Insights::FTreeNodeGroupingByUniqueValueCString>() &&
-						Grouping->As<UE::Insights::FTreeNodeGroupingByUniqueValueCString>().GetColumnId() == FAssetTableColumns::PrimaryTypeColumnId;
+						Grouping->As<UE::Insights::FTreeNodeGroupingByUniqueValueCString>().GetColumnId() == FAssetTableColumns::TypeColumnId;
 				});
 			if (PrimaryTypeGrouping)
 			{
@@ -311,22 +312,22 @@
 			InOutConfigSet.Add({ FAssetTableColumns::PluginNameColumnId,                    true, 200.0f });
 		}
 	};
-	AvailableViewPresets.Add(MakeShared<FPrimaryTypeViewPreset>());
+	AvailableViewPresets.Add(MakeShared<FAssetTypeViewPreset>());
 
+	//////////////////////////////////////////////////
+	// Plugin Group View
 
-	////////////////////////////////////////////////////
-	//// Class Type Breakdown View
-
-	class FClassTypeViewPreset : public UE::Insights::ITableTreeViewPreset
+	class FPluginView : public UE::Insights::ITableTreeViewPreset
 	{
 	public:
 		virtual FText GetName() const override
 		{
-			return LOCTEXT("ClassType_PresetName", "Group By Class Type");
+			return LOCTEXT("PluginDepView_PresetName", "Group By Plugin");
 		}
+
 		virtual FText GetToolTip() const override
 		{
-			return LOCTEXT("ClassType_PresetToolTip", "Class Type Breakdown View\nConfigure the tree view to show a breakdown of assets by their class type.");
+			return LOCTEXT("PluginDepView_PresetToolTip", "Group By Plugin\nGroup assets by plugin.");
 		}
 		virtual FName GetSortColumn() const override
 		{
@@ -343,26 +344,26 @@
 			check(InAvailableGroupings[0]->Is<UE::Insights::FTreeNodeGroupingFlat>());
 			InOutCurrentGroupings.Add(InAvailableGroupings[0]);
 
-			const TSharedPtr<UE::Insights::FTreeNodeGrouping>* PrimaryTypeGrouping = InAvailableGroupings.FindByPredicate(
+			const TSharedPtr<UE::Insights::FTreeNodeGrouping>* PrimaryGrouping = InAvailableGroupings.FindByPredicate(
 				[](TSharedPtr<UE::Insights::FTreeNodeGrouping>& Grouping)
 				{
 					return Grouping->Is<UE::Insights::FTreeNodeGroupingByUniqueValueCString>() &&
-						Grouping->As<UE::Insights::FTreeNodeGroupingByUniqueValueCString>().GetColumnId() == FAssetTableColumns::NativeClassColumnId;
-				});
-			if (PrimaryTypeGrouping)
-			{
-				InOutCurrentGroupings.Add(*PrimaryTypeGrouping);
-			}
-
-			const TSharedPtr<UE::Insights::FTreeNodeGrouping>* SecondaryTypeGrouping = InAvailableGroupings.FindByPredicate(
-				[](TSharedPtr<UE::Insights::FTreeNodeGrouping>& Grouping)
-				{
-					return Grouping->Is<UE::Insights::FTreeNodeGroupingByUniqueValueCString>() &&
 						Grouping->As<UE::Insights::FTreeNodeGroupingByUniqueValueCString>().GetColumnId() == FAssetTableColumns::PluginNameColumnId;
 				});
-			if (SecondaryTypeGrouping)
+			if (PrimaryGrouping)
 			{
-				InOutCurrentGroupings.Add(*SecondaryTypeGrouping);
+				InOutCurrentGroupings.Add(*PrimaryGrouping);
 			}
+
+			const TSharedPtr<UE::Insights::FTreeNodeGrouping>* SecondaryGrouping = InAvailableGroupings.FindByPredicate(
+				[](TSharedPtr<UE::Insights::FTreeNodeGrouping>& Grouping)
+				{
+					return Grouping->Is<UE::Insights::FTreeNodeGroupingByUniqueValueCString>() &&
+						Grouping->As<UE::Insights::FTreeNodeGroupingByUniqueValueCString>().GetColumnId() == FAssetTableColumns::TypeColumnId;
+				});
+			if (SecondaryGrouping)
+			{
+				InOutCurrentGroupings.Add(*SecondaryGrouping);
+			}
 		}
 		virtual void GetColumnConfigSet(TArray<UE::Insights::FTableColumnConfig>& InOutConfigSet) const override
@@ -384,21 +385,22 @@
 			InOutConfigSet.Add({ FAssetTableColumns::PluginNameColumnId,                    true, 200.0f });
 		}
 	};
-	AvailableViewPresets.Add(MakeShared<FClassTypeViewPreset>());
+	AvailableViewPresets.Add(MakeShared<FPluginView>());
 
-	////////////////////////////////////////////////////
-	//// Asset Type Breakdown View
+	//////////////////////////////////////////////////
+	// Plugin Dependency View
 
-	class FAssetTypeViewPreset : public UE::Insights::ITableTreeViewPreset
+	class FPluginDependencyView : public UE::Insights::ITableTreeViewPreset
 	{
 	public:
 		virtual FText GetName() const override
 		{
-			return LOCTEXT("AssetType_PresetName", "Group By Asset Type");
+			return LOCTEXT("PluginDepView_PresetName", "Plugin Dependency Analysis");
 		}
+
 		virtual FText GetToolTip() const override
 		{
-			return LOCTEXT("AssetType_PresetToolTip", "Asset Type Breakdown View\nConfigure the tree view to show a breakdown of assets by their asset type.");
+			return LOCTEXT("PluginDepView_PresetToolTip", "Plugin Dependency Analysis View\nConfigure the tree view to show a breakdown of assets by Game Feature Plugin, showing also the dependencies between plugins.");
 		}
 		virtual FName GetSortColumn() const override
 		{
@@ -415,26 +417,25 @@
 			check(InAvailableGroupings[0]->Is<UE::Insights::FTreeNodeGroupingFlat>());
 			InOutCurrentGroupings.Add(InAvailableGroupings[0]);
 
-			const TSharedPtr<UE::Insights::FTreeNodeGrouping>* PrimaryTypeGrouping = InAvailableGroupings.FindByPredicate(
+			const TSharedPtr<UE::Insights::FTreeNodeGrouping>* PrimaryGrouping = InAvailableGroupings.FindByPredicate(
+				[](TSharedPtr<UE::Insights::FTreeNodeGrouping>& Grouping)
+				{
+					return Grouping->Is<FPluginDependencyGrouping>();
+				});
+			if (PrimaryGrouping)
+			{
+				InOutCurrentGroupings.Add(*PrimaryGrouping);
+			}
+
+			const TSharedPtr<UE::Insights::FTreeNodeGrouping>* SecondaryGrouping = InAvailableGroupings.FindByPredicate(
 				[](TSharedPtr<UE::Insights::FTreeNodeGrouping>& Grouping)
 				{
 					return Grouping->Is<UE::Insights::FTreeNodeGroupingByUniqueValueCString>() &&
 						Grouping->As<UE::Insights::FTreeNodeGroupingByUniqueValueCString>().GetColumnId() == FAssetTableColumns::TypeColumnId;
 				});
-			if (PrimaryTypeGrouping)
+			if (SecondaryGrouping)
 			{
-				InOutCurrentGroupings.Add(*PrimaryTypeGrouping);
+				InOutCurrentGroupings.Add(*SecondaryGrouping);
 			}
-
-			const TSharedPtr<UE::Insights::FTreeNodeGrouping>* SecondaryTypeGrouping = InAvailableGroupings.FindByPredicate(
-				[](TSharedPtr<UE::Insights::FTreeNodeGrouping>& Grouping)
-				{
-					return Grouping->Is<UE::Insights::FTreeNodeGroupingByUniqueValueCString>() &&
-						Grouping->As<UE::Insights::FTreeNodeGroupingByUniqueValueCString>().GetColumnId() == FAssetTableColumns::PluginNameColumnId;
-				});
-			if (SecondaryTypeGrouping)
-			{
-				InOutCurrentGroupings.Add(*SecondaryTypeGrouping);
-			}
 		}
 		virtual void GetColumnConfigSet(TArray<UE::Insights::FTableColumnConfig>& InOutConfigSet) const override
@@ -456,22 +457,21 @@
 			InOutConfigSet.Add({ FAssetTableColumns::PluginNameColumnId,                    true, 200.0f });
 		}
 	};
-	AvailableViewPresets.Add(MakeShared<FAssetTypeViewPreset>());
+	AvailableViewPresets.Add(MakeShared<FPluginDependencyView>());
 
-	//////////////////////////////////////////////////
-	// Plugin Group View
+	////////////////////////////////////////////////////
+	//// Primary Type Breakdown View
 
-	class FPluginView : public UE::Insights::ITableTreeViewPreset
+	class FPrimaryTypeViewPreset : public UE::Insights::ITableTreeViewPreset
 	{
 	public:
 		virtual FText GetName() const override
 		{
-			return LOCTEXT("PluginDepView_PresetName", "Group By Plugin");
+			return LOCTEXT("PrimaryType_PresetName", "Group By Primary Type");
 		}
-
 		virtual FText GetToolTip() const override
 		{
-			return LOCTEXT("PluginDepView_PresetToolTip", "Group By Plugin\nGroup assets by plugin.");
+			return LOCTEXT("PrimaryType_PresetToolTip", "Primary Type Breakdown View\nConfigure the tree view to show a breakdown of assets by their primary type.");
 		}
 		virtual FName GetSortColumn() const override
 		{
@@ -488,26 +488,26 @@
 			check(InAvailableGroupings[0]->Is<UE::Insights::FTreeNodeGroupingFlat>());
 			InOutCurrentGroupings.Add(InAvailableGroupings[0]);
 
-			const TSharedPtr<UE::Insights::FTreeNodeGrouping>* PrimaryGrouping = InAvailableGroupings.FindByPredicate(
+			const TSharedPtr<UE::Insights::FTreeNodeGrouping>* PrimaryTypeGrouping = InAvailableGroupings.FindByPredicate(
 				[](TSharedPtr<UE::Insights::FTreeNodeGrouping>& Grouping)
 				{
 					return Grouping->Is<UE::Insights::FTreeNodeGroupingByUniqueValueCString>() &&
+						Grouping->As<UE::Insights::FTreeNodeGroupingByUniqueValueCString>().GetColumnId() == FAssetTableColumns::PrimaryTypeColumnId;
+				});
+			if (PrimaryTypeGrouping)
+			{
+				InOutCurrentGroupings.Add(*PrimaryTypeGrouping);
+			}
+
+			const TSharedPtr<UE::Insights::FTreeNodeGrouping>* SecondaryTypeGrouping = InAvailableGroupings.FindByPredicate(
+				[](TSharedPtr<UE::Insights::FTreeNodeGrouping>& Grouping)
+				{
+					return Grouping->Is<UE::Insights::FTreeNodeGroupingByUniqueValueCString>() &&
 						Grouping->As<UE::Insights::FTreeNodeGroupingByUniqueValueCString>().GetColumnId() == FAssetTableColumns::PluginNameColumnId;
 				});
-			if (PrimaryGrouping)
+			if (SecondaryTypeGrouping)
 			{
-				InOutCurrentGroupings.Add(*PrimaryGrouping);
+				InOutCurrentGroupings.Add(*SecondaryTypeGrouping);
 			}
-
-			const TSharedPtr<UE::Insights::FTreeNodeGrouping>* SecondaryGrouping = InAvailableGroupings.FindByPredicate(
-				[](TSharedPtr<UE::Insights::FTreeNodeGrouping>& Grouping)
-				{
-					return Grouping->Is<UE::Insights::FTreeNodeGroupingByUniqueValueCString>() &&
-						Grouping->As<UE::Insights::FTreeNodeGroupingByUniqueValueCString>().GetColumnId() == FAssetTableColumns::TypeColumnId;
-				});
-			if (SecondaryGrouping)
-			{
-				InOutCurrentGroupings.Add(*SecondaryGrouping);
-			}
 		}
 		virtual void GetColumnConfigSet(TArray<UE::Insights::FTableColumnConfig>& InOutConfigSet) const override
@@ -529,22 +529,22 @@
 			InOutConfigSet.Add({ FAssetTableColumns::PluginNameColumnId,                    true, 200.0f });
 		}
 	};
-	AvailableViewPresets.Add(MakeShared<FPluginView>());
+	AvailableViewPresets.Add(MakeShared<FPrimaryTypeViewPreset>());
 
-	//////////////////////////////////////////////////
-	// Plugin Dependency View
 
-	class FPluginDependencyView : public UE::Insights::ITableTreeViewPreset
+	////////////////////////////////////////////////////
+	//// Class Type Breakdown View
+
+	class FClassTypeViewPreset : public UE::Insights::ITableTreeViewPreset
 	{
 	public:
 		virtual FText GetName() const override
 		{
-			return LOCTEXT("PluginDepView_PresetName", "Plugin Dependency Analysis");
+			return LOCTEXT("ClassType_PresetName", "Group By Class Type");
 		}
-
 		virtual FText GetToolTip() const override
 		{
-			return LOCTEXT("PluginDepView_PresetToolTip", "Plugin Dependency Analysis View\nConfigure the tree view to show a breakdown of assets by Game Feature Plugin, showing also the dependencies between plugins.");
+			return LOCTEXT("ClassType_PresetToolTip", "Class Type Breakdown View\nConfigure the tree view to show a breakdown of assets by their class type.");
 		}
 		virtual FName GetSortColumn() const override
 		{
@@ -561,25 +561,26 @@
 			check(InAvailableGroupings[0]->Is<UE::Insights::FTreeNodeGroupingFlat>());
 			InOutCurrentGroupings.Add(InAvailableGroupings[0]);
 
-			const TSharedPtr<UE::Insights::FTreeNodeGrouping>* PrimaryGrouping = InAvailableGroupings.FindByPredicate(
-				[](TSharedPtr<UE::Insights::FTreeNodeGrouping>& Grouping)
-				{
-					return Grouping->Is<FPluginDependencyGrouping>();
-				});
-			if (PrimaryGrouping)
-			{
-				InOutCurrentGroupings.Add(*PrimaryGrouping);
-			}
-
-			const TSharedPtr<UE::Insights::FTreeNodeGrouping>* SecondaryGrouping = InAvailableGroupings.FindByPredicate(
+			const TSharedPtr<UE::Insights::FTreeNodeGrouping>* PrimaryTypeGrouping = InAvailableGroupings.FindByPredicate(
 				[](TSharedPtr<UE::Insights::FTreeNodeGrouping>& Grouping)
 				{
 					return Grouping->Is<UE::Insights::FTreeNodeGroupingByUniqueValueCString>() &&
-						Grouping->As<UE::Insights::FTreeNodeGroupingByUniqueValueCString>().GetColumnId() == FAssetTableColumns::TypeColumnId;
+						Grouping->As<UE::Insights::FTreeNodeGroupingByUniqueValueCString>().GetColumnId() == FAssetTableColumns::NativeClassColumnId;
 				});
-			if (SecondaryGrouping)
+			if (PrimaryTypeGrouping)
 			{
-				InOutCurrentGroupings.Add(*SecondaryGrouping);
+				InOutCurrentGroupings.Add(*PrimaryTypeGrouping);
 			}
+
+			const TSharedPtr<UE::Insights::FTreeNodeGrouping>* SecondaryTypeGrouping = InAvailableGroupings.FindByPredicate(
+				[](TSharedPtr<UE::Insights::FTreeNodeGrouping>& Grouping)
+				{
+					return Grouping->Is<UE::Insights::FTreeNodeGroupingByUniqueValueCString>() &&
+						Grouping->As<UE::Insights::FTreeNodeGroupingByUniqueValueCString>().GetColumnId() == FAssetTableColumns::PluginNameColumnId;
+				});
+			if (SecondaryTypeGrouping)
+			{
+				InOutCurrentGroupings.Add(*SecondaryTypeGrouping);
+			}
 		}
 		virtual void GetColumnConfigSet(TArray<UE::Insights::FTableColumnConfig>& InOutConfigSet) const override
@@ -601,12 +602,9 @@
 			InOutConfigSet.Add({ FAssetTableColumns::PluginNameColumnId,                    true, 200.0f });
 		}
 	};
-	AvailableViewPresets.Add(MakeShared<FPluginDependencyView>());
+	AvailableViewPresets.Add(MakeShared<FClassTypeViewPreset>());
 
-	
-
 	//////////////////////////////////////////////////
-
 	SelectedViewPreset = AvailableViewPresets[0];
 }
 
