Added icon for Dependenicies

diff --git a/Engine/Plugins/Editor/AssetManagerEditor/Source/AssetManagerEditor/Private/Insights/Common/InsightsStyle.cpp b/Engine/Plugins/Editor/AssetManagerEditor/Source/AssetManagerEditor/Private/Insights/Common/InsightsStyle.cpp
--- a/Engine/Plugins/Editor/AssetManagerEditor/Source/AssetManagerEditor/Private/Insights/Common/InsightsStyle.cpp
+++ b/Engine/Plugins/Editor/AssetManagerEditor/Source/AssetManagerEditor/Private/Insights/Common/InsightsStyle.cpp
@@ -166,6 +166,7 @@ void FInsightsStyle::FStyle::Initialize()
 	Set("Icons.Asset.TreeItem", new CORE_IMAGE_BRUSH_SVG("Starship/Common/box-perspective", Icon12x12));
 	Set("Icons.Package.TreeItem", new EDITOR_IMAGE_BRUSH_SVG("Starship/Common/ProjectPackage", Icon12x12));
 	Set("Icons.Plugin.TreeItem", new EDITOR_IMAGE_BRUSH_SVG("Starship/Common/Plugins", Icon12x12));
+	Set("Icons.Dependencies.TreeItem", new EDITOR_IMAGE_BRUSH_SVG("Starship/Common/Blueprint", Icon12x12));
 
 	Set("TreeTable.TooltipBold", FTextBlockStyle(NormalText)
 		.SetFont(DEFAULT_FONT("Bold", 8))
diff --git a/Engine/Plugins/Editor/AssetManagerEditor/Source/AssetManagerEditor/Private/TreeView/AssetTreeNode.cpp b/Engine/Plugins/Editor/AssetManagerEditor/Source/AssetManagerEditor/Private/TreeView/AssetTreeNode.cpp
--- a/Engine/Plugins/Editor/AssetManagerEditor/Source/AssetManagerEditor/Private/TreeView/AssetTreeNode.cpp
+++ b/Engine/Plugins/Editor/AssetManagerEditor/Source/AssetManagerEditor/Private/TreeView/AssetTreeNode.cpp
@@ -43,6 +43,12 @@ const FSlateBrush* FAssetTreeNode::GetIcon(EStyle Style) const
 			break;
 		}	
 
+		case EStyle::EDependencies:
+		{
+			return UE::Insights::FInsightsStyle::GetBrush("Icons.Dependencies.TreeItem");
+			break;
+		}
+
 		case EStyle::EPlugin:
 		{
 			return UE::Insights::FInsightsStyle::GetBrush("Icons.Plugin.TreeItem");
@@ -75,6 +81,12 @@ FLinearColor FAssetTreeNode::GetColor(EStyle Style) const
 			break;
 		}
 
+		case EStyle::EDependencies:
+		{
+			return FLinearColor(1.0f, 1.0f, 1.0f, 1.0f);
+			break;
+		}
+
 		case EStyle::EPlugin:
 		{
 			return FLinearColor(1.0f, 1.0f, 1.0f, 1.0f);
@@ -313,7 +325,7 @@ bool FPluginDependenciesGroupNode::OnLazyCreateChildren(TSharedPtr<class UE::Ins
 
 FAssetTreeNode::EStyle FPluginDependenciesGroupNode::GetStyle() const
 {
-	return EStyle::EGroup;
+	return EStyle::EDependencies;
 }
 
 ////////////////////////////////////////////////////////////////////////////////////////////////////
diff --git a/Engine/Plugins/Editor/AssetManagerEditor/Source/AssetManagerEditor/Private/TreeView/AssetTreeNode.h b/Engine/Plugins/Editor/AssetManagerEditor/Source/AssetManagerEditor/Private/TreeView/AssetTreeNode.h
--- a/Engine/Plugins/Editor/AssetManagerEditor/Source/AssetManagerEditor/Private/TreeView/AssetTreeNode.h
+++ b/Engine/Plugins/Editor/AssetManagerEditor/Source/AssetManagerEditor/Private/TreeView/AssetTreeNode.h
@@ -54,6 +54,7 @@ protected:
 		EDefault,
 		EGroup,
 		EAsset,
+		EDependencies,
 		EPlugin,
 	};
 
