LOCTEXT build fix.

diff --git a/Engine/Plugins/Editor/AssetManagerEditor/Source/AssetManagerEditor/Private/TreeView/SAssetTableTreeView.cpp b/Engine/Plugins/Editor/AssetManagerEditor/Source/AssetManagerEditor/Private/TreeView/SAssetTableTreeView.cpp
--- a/Engine/Plugins/Editor/AssetManagerEditor/Source/AssetManagerEditor/Private/TreeView/SAssetTableTreeView.cpp
+++ b/Engine/Plugins/Editor/AssetManagerEditor/Source/AssetManagerEditor/Private/TreeView/SAssetTableTreeView.cpp
@@ -322,12 +322,12 @@ void SAssetTableTreeView::InitAvailableViewPresets()
 	public:
 		virtual FText GetName() const override
 		{
-			return LOCTEXT("PluginDepView_PresetName", "Group By Plugin");
+			return LOCTEXT("PluginGroupView_PresetName", "Group By Plugin");
 		}
 
 		virtual FText GetToolTip() const override
 		{
-			return LOCTEXT("PluginDepView_PresetToolTip", "Group By Plugin\nGroup assets by plugin.");
+			return LOCTEXT("PluginGroupView_PresetToolTip", "Group By Plugin\nGroup assets by plugin.");
 		}
 		virtual FName GetSortColumn() const override
 		{
