Added plugin inclusive size calculations and display. Various UX improvements including:
* Made all plugins the plugin dependency analysis view show up in a uniform way (plugin icon + Assets and Dependencies child nodes)
* Hid various columns by default and changed ordering
* Restored Asset Dependency Analysis preset
* Changed preset order
* Improved icons for asset dependencies

diff --git a/Engine/Plugins/Editor/AssetManagerEditor/Source/AssetManagerEditor/Private/TreeView/AssetDependencyGrouping.cpp b/Engine/Plugins/Editor/AssetManagerEditor/Source/AssetManagerEditor/Private/TreeView/AssetDependencyGrouping.cpp
--- a/Engine/Plugins/Editor/AssetManagerEditor/Source/AssetManagerEditor/Private/TreeView/AssetDependencyGrouping.cpp
+++ b/Engine/Plugins/Editor/AssetManagerEditor/Source/AssetManagerEditor/Private/TreeView/AssetDependencyGrouping.cpp
@@ -56,10 +56,6 @@ void FAssetDependencyGrouping::GroupNodes(const TArray<UE::Insights::FTableTreeN
 
 		// For each (visible) asset, we will create the following hierarchy:
 
-		// If the asset does not have dependencies, the asset node is added directly (no group node is created).
-		// |
-		// +-- [asset:{AssetName}]
-
 		// By default, the "Dependencies" node is collapsed.
 		// |
 		// +-- [group:{AssetName}] (self + dependencies)
@@ -240,7 +236,7 @@ void FPluginDependencyGrouping::GroupNodes(const TArray<UE::Insights::FTableTree
 			// Plugin group node (simple or with dependencies) was already created. Just add the current asset to it.
 			PluginGroup->AddChildAndSetParent(NodePtr);
 		}
-		else if (PluginInfo.PluginDependencies.Num() > 0)
+		else
 		{
 			// Create the Plugin Self+Dependencies group node and add the current asset to the Self group.
 			FName PluginAndDependenciesGroupName = AssetTable->GetNameForPlugin(PluginIndex);
@@ -250,16 +246,6 @@ void FPluginDependencyGrouping::GroupNodes(const TArray<UE::Insights::FTableTree
 			PluginGroup->AddChildAndSetParent(NodePtr);
 			PluginIndexToGroupNodeMap.Add(PluginIndex, PluginGroup);
 		}
-		else
-		{
-			// Create a simple Plugin group node and add the current asset to it.
-			FName PluginGroupName = AssetTable->GetNameForPlugin(PluginIndex);
-			TSharedPtr<FPluginSimpleGroupNode> PluginGroupPtr = MakeShared<FPluginSimpleGroupNode>(PluginGroupName, AssetTable, PluginIndex);
-			ParentGroup.AddChildAndSetParent(PluginGroupPtr);
-			PluginGroup = PluginGroupPtr.Get();
-			PluginGroup->AddChildAndSetParent(NodePtr);
-			PluginIndexToGroupNodeMap.Add(PluginIndex, PluginGroup);
-		}
 	}
 }
 
diff --git a/Engine/Plugins/Editor/AssetManagerEditor/Source/AssetManagerEditor/Private/TreeView/AssetTable.cpp b/Engine/Plugins/Editor/AssetManagerEditor/Source/AssetManagerEditor/Private/TreeView/AssetTable.cpp
--- a/Engine/Plugins/Editor/AssetManagerEditor/Source/AssetManagerEditor/Private/TreeView/AssetTable.cpp
+++ b/Engine/Plugins/Editor/AssetManagerEditor/Source/AssetManagerEditor/Private/TreeView/AssetTable.cpp
@@ -35,6 +35,7 @@ const FName FAssetTableColumns::TotalUsageCountColumnId(TEXT("TotalUsageCount"))
 const FName FAssetTableColumns::ChunksColumnId(TEXT("Chunks"));
 const FName FAssetTableColumns::NativeClassColumnId(TEXT("NativeClass"));
 const FName FAssetTableColumns::PluginNameColumnId(TEXT("PluginName"));
+const FName FAssetTableColumns::PluginInclusiveSizeColumnId(TEXT("PluginInclusiveSize"));
 
 ////////////////////////////////////////////////////////////////////////////////////////////////////
 // FAssetTableStringValueGetterWithDependencyAggregationHandling
@@ -824,7 +825,7 @@ void FAssetTable::AddDefaultColumns()
 		Column.SetTitleName(LOCTEXT("PrimaryTypeColumnTitle", "Primary Type"));
 		Column.SetDescription(LOCTEXT("PrimaryTypeColumnDesc", "Primary Asset Type of this asset, if set"));
 
-		Column.SetFlags(ETableColumnFlags::ShouldBeVisible | ETableColumnFlags::CanBeHidden | ETableColumnFlags::CanBeFiltered);
+		Column.SetFlags(ETableColumnFlags::CanBeHidden | ETableColumnFlags::CanBeFiltered);
 
 		Column.SetHorizontalAlignment(HAlign_Left);
 		Column.SetInitialWidth(200.0f);
@@ -866,7 +867,7 @@ void FAssetTable::AddDefaultColumns()
 		Column.SetTitleName(LOCTEXT("PrimaryNameColumnTitle", "Primary Name"));
 		Column.SetDescription(LOCTEXT("PrimaryNameColumnDesc", "Primary Asset Name of this asset, if set"));
 
-		Column.SetFlags(ETableColumnFlags::ShouldBeVisible | ETableColumnFlags::CanBeHidden | ETableColumnFlags::CanBeFiltered);
+		Column.SetFlags(ETableColumnFlags::CanBeHidden | ETableColumnFlags::CanBeFiltered);
 
 		Column.SetHorizontalAlignment(HAlign_Left);
 		Column.SetInitialWidth(200.0f);
@@ -1107,7 +1108,7 @@ void FAssetTable::AddDefaultColumns()
 		Column.SetTitleName(LOCTEXT("TotalUsageCountColumnTitle", "Total Usage Count"));
 		Column.SetDescription(LOCTEXT("TotalUsageCountColumnDesc", "Weighted count of Primary Assets that use this\nA higher usage means it's more likely to be in memory at runtime."));
 
-		Column.SetFlags(ETableColumnFlags::ShouldBeVisible | ETableColumnFlags::CanBeHidden | ETableColumnFlags::CanBeFiltered | ETableColumnFlags::IsDynamic);
+		Column.SetFlags(ETableColumnFlags::CanBeHidden | ETableColumnFlags::CanBeFiltered | ETableColumnFlags::IsDynamic);
 
 		Column.SetHorizontalAlignment(HAlign_Right);
 		Column.SetInitialWidth(100.0f);
@@ -1163,10 +1164,10 @@ void FAssetTable::AddDefaultColumns()
 		Column.SetTitleName(LOCTEXT("ChunksColumnTitle", "Chunks"));
 		Column.SetDescription(LOCTEXT("ChunksColumnDesc", "List of chunks this will be added to when cooked"));
 
-		Column.SetFlags(ETableColumnFlags::ShouldBeVisible | ETableColumnFlags::CanBeHidden | ETableColumnFlags::CanBeFiltered);
+		Column.SetFlags(ETableColumnFlags::CanBeHidden | ETableColumnFlags::CanBeFiltered);
 
 		Column.SetHorizontalAlignment(HAlign_Left);
-		Column.SetInitialWidth(200.0f);
+		Column.SetInitialWidth(100.0f);
 
 		Column.SetDataType(ETableCellDataType::CString);
 
@@ -1205,7 +1206,7 @@ void FAssetTable::AddDefaultColumns()
 		Column.SetTitleName(LOCTEXT("NativeClassColumnTitle", "Native Class"));
 		Column.SetDescription(LOCTEXT("NativeClassColumnDesc", "Native class of the asset"));
 
-		Column.SetFlags(ETableColumnFlags::ShouldBeVisible | ETableColumnFlags::CanBeHidden | ETableColumnFlags::CanBeFiltered);
+		Column.SetFlags(ETableColumnFlags::CanBeHidden | ETableColumnFlags::CanBeFiltered);
 
 		Column.SetHorizontalAlignment(HAlign_Left);
 		Column.SetInitialWidth(200.0f);
@@ -1278,6 +1279,54 @@ void FAssetTable::AddDefaultColumns()
 		AddColumn(ColumnRef);
 	}
 	//////////////////////////////////////////////////
+	// PluginInclusiveSize Column
+	{
+		TSharedRef<FTableColumn> ColumnRef = MakeShared<FTableColumn>(FAssetTableColumns::PluginInclusiveSizeColumnId);
+		FTableColumn& Column = *ColumnRef;
+
+		Column.SetIndex(ColumnIndex++);
+
+		Column.SetShortName(LOCTEXT("PluginInclusiveSizeColumnName", "Incl. Size Plugin"));
+		Column.SetTitleName(LOCTEXT("PluginInclusiveSizeColumnTitle", "Plugin Inclusive Size"));
+		Column.SetDescription(LOCTEXT("PluginInclusiveColumnDesc", "Inclusive size of this plugin and its dependencies"));
+
+		Column.SetFlags(ETableColumnFlags::ShouldBeVisible | ETableColumnFlags::CanBeHidden | ETableColumnFlags::CanBeFiltered);
+
+		Column.SetHorizontalAlignment(HAlign_Left);
+		Column.SetInitialWidth(50.0f);
+
+		Column.SetDataType(ETableCellDataType::Int64);
+
+		class FPluginInclusiveSizeValueGetter : public FTableCellValueGetter
+		{
+		public:
+			virtual const TOptional<FTableCellValue> GetValue(const FTableColumn& Column, const FBaseTreeNode& Node) const override
+			{
+				if (Node.Is<FPluginSimpleGroupNode>())
+				{
+					TSharedPtr<FTable> TablePtr = static_cast<const FPluginSimpleGroupNode&>(Node).GetParentTable().Pin();
+					const FAssetTable& AssetTable = static_cast<const FAssetTable&>(*TablePtr);
+					int32 PluginIndex = static_cast<const FPluginSimpleGroupNode&>(Node).GetPluginIndex();
+					return FTableCellValue(AssetTable.GetPluginInfoByIndex(PluginIndex).GetOrComputeTotalSizeInclusiveOfDependencies(AssetTable));
+				}
+
+				return TOptional<FTableCellValue>();
+			}
+		};
+		TSharedRef<ITableCellValueGetter> Getter = MakeShared<FPluginInclusiveSizeValueGetter>();
+		Column.SetValueGetter(Getter);
+
+		TSharedRef<ITableCellValueFormatter> Formatter = MakeShared<FInt64ValueFormatterAsMemory>();
+		Column.SetValueFormatter(Formatter);
+
+		TSharedRef<ITableCellValueSorter> Sorter = MakeShared<FSorterByInt64Value>(ColumnRef);
+		Column.SetValueSorter(Sorter);
+		Column.SetInitialSortMode(EColumnSortMode::Descending);
+
+		Column.SetAggregation(ETableColumnAggregation::None);
+
+		AddColumn(ColumnRef);
+	}
 }
 
 /*static*/TSet<int32> FAssetTableRow::GatherAllReachableNodes(const TArray<int32>& StartingNodes, const FAssetTable& OwningTable, const TSet<int32>& AdditionalNodesToStopAt, const TSet<const TCHAR*, TStringPointerSetKeyFuncs_DEPRECATED<const TCHAR*>>& RestrictToPlugins, TMap<int32, TArray<int32>>* OutRouteMap /*= nullptr*/)
@@ -1605,4 +1654,134 @@ int64 FAssetTableRow::GetOrComputeTotalSizeExternalDependencies(const FAssetTabl
 
 ////////////////////////////////////////////////////////////////////////////////////////////////////
 
+int64 FAssetTablePluginInfo::GetOrComputeTotalSizeInclusiveOfDependencies(const FAssetTable& OwningTable) const
+{
+	if (InclusiveSize == -1)
+	{
+		InclusiveSize = Size;
+
+		TArray<int32> DependencyStack = PluginDependencies;
+		TSet<int32> AllDependencies;
+
+		// Gather all the dependencies
+		while (DependencyStack.Num() > 0)
+		{
+			int32 CurrentIndex = DependencyStack.Pop();
+			if (AllDependencies.Contains(CurrentIndex))
+			{
+				continue;
+			}
+			AllDependencies.Add(CurrentIndex);
+			const FAssetTablePluginInfo& PluginInfo = OwningTable.GetPluginInfoByIndex(CurrentIndex);
+			InclusiveSize += PluginInfo.GetSize();
+			DependencyStack.Append(PluginInfo.GetDependencies());
+		}
+	}
+	return InclusiveSize;
+}
+
+////////////////////////////////////////////////////////////////////////////////////////////////////
+
+int64 FAssetTablePluginInfo::GetOrComputeTotalSizeUniqueDependencies(const FAssetTable& OwningTable) const
+{
+	if (UniqueDependenciesSize == -1)
+	{
+		ComputeDependencySizes(OwningTable);
+	}
+	return UniqueDependenciesSize;
+}
+
+////////////////////////////////////////////////////////////////////////////////////////////////////
+
+int64 FAssetTablePluginInfo::GetOrComputeTotalSizeSharedDependencies(const FAssetTable& OwningTable) const
+{
+	if (SharedDependenciesSize == -1)
+	{
+		ComputeDependencySizes(OwningTable);
+	}
+	return SharedDependenciesSize;
+}
+
+////////////////////////////////////////////////////////////////////////////////////////////////////
+
+void FAssetTablePluginInfo::ComputeDependencySizes(const FAssetTable& OwningTable) const
+{
+	SharedDependenciesSize = 0;
+	UniqueDependenciesSize = 0;
+
+	// Plugins can't have circular dependencies, so our algorithm for exploring the dependencies is significantly simpler than for assets
+
+	const int32 RootIndex = OwningTable.GetIndexForPlugin(PluginName);
+
+	TArray<int32> DependencyStack = PluginDependencies;
+	TSet<int32> AllDependencies;
+
+	// Gather all the dependencies
+	while (DependencyStack.Num() > 0)
+	{
+		int32 CurrentIndex = DependencyStack.Pop();
+		AllDependencies.Add(CurrentIndex);
+		DependencyStack.Append(OwningTable.GetPluginInfoByIndex(CurrentIndex).GetDependencies());
+	}
+
+	// Find the set of dependencies which are directly referenced by something outside this set of plugins
+	TSet<int32> SharedDependencies;
+	for (int32 DependencyIndex : AllDependencies)
+	{
+		const FAssetTablePluginInfo& CurrentPluginInfo = OwningTable.GetPluginInfoByIndex(DependencyIndex);
+		bool IsShared = false;
+		for (int32 ReferencerIndex : CurrentPluginInfo.GetReferencers())
+		{
+			bool ReferencerIsRoot = (ReferencerIndex == RootIndex);
+			bool ReferencerIsNotInDependencyTree = !AllDependencies.Contains(ReferencerIndex);
+			if (ReferencerIsNotInDependencyTree && !ReferencerIsRoot)
+			{
+				IsShared = true;
+				break;
+			}
+		}
+		if (IsShared)
+		{
+			SharedDependencies.Add(DependencyIndex);
+		}
+	}
+
+	// Add to the SharedDependencies list all transitive dependencies
+	TArray<int32> SharedDependencyStack = SharedDependencies.Array();
+	while (SharedDependencyStack.Num() > 0)
+	{
+		int32 CurrentIndex = SharedDependencyStack.Pop();
+		SharedDependencies.Add(CurrentIndex);
+		SharedDependencyStack.Append(OwningTable.GetPluginInfoByIndex(CurrentIndex).GetDependencies());
+	}
+
+	// UniqueDependencies are all the dependencies that aren't shared
+	TArray<int32> UniqueDependencies = AllDependencies.Array();
+	for (int32 SharedDependencyIndex : SharedDependencies)
+	{
+		UniqueDependencies.RemoveSwap(SharedDependencyIndex);
+	}
+	
+	ensure((UniqueDependencies.Num() + SharedDependencies.Num()) == AllDependencies.Num());
+
+	// Add up the sizes
+	for (int32 UniqueDependencyIndex : UniqueDependencies)
+	{
+		const FAssetTablePluginInfo& CurrentPluginInfo = OwningTable.GetPluginInfoByIndexChecked(UniqueDependencyIndex);
+		int64 PluginSize = CurrentPluginInfo.GetSize();
+		ensureMsgf(PluginSize >= 0, TEXT("Found plugin %s with uninitialized size."), CurrentPluginInfo.GetName());
+		UniqueDependenciesSize += PluginSize;
+	}
+
+	for (int32 SharedDependencyIndex : SharedDependencies)
+	{
+		const FAssetTablePluginInfo& CurrentPluginInfo = OwningTable.GetPluginInfoByIndexChecked(SharedDependencyIndex);
+		int64 PluginSize = CurrentPluginInfo.GetSize();
+		ensureMsgf(PluginSize >= 0, TEXT("Found plugin %s with uninitialized size."), CurrentPluginInfo.GetName());
+		SharedDependenciesSize += PluginSize;
+	}
+}
+
+////////////////////////////////////////////////////////////////////////////////////////////////////
+
 #undef LOCTEXT_NAMESPACE
diff --git a/Engine/Plugins/Editor/AssetManagerEditor/Source/AssetManagerEditor/Private/TreeView/AssetTable.h b/Engine/Plugins/Editor/AssetManagerEditor/Source/AssetManagerEditor/Private/TreeView/AssetTable.h
--- a/Engine/Plugins/Editor/AssetManagerEditor/Source/AssetManagerEditor/Private/TreeView/AssetTable.h
+++ b/Engine/Plugins/Editor/AssetManagerEditor/Source/AssetManagerEditor/Private/TreeView/AssetTable.h
@@ -125,6 +125,7 @@ struct FAssetTableColumns
 	static const FName ChunksColumnId;
 	static const FName NativeClassColumnId;
 	static const FName PluginNameColumnId;
+	static const FName PluginInclusiveSizeColumnId;
 };
 
 struct FAssetTableDependencySizes
@@ -133,12 +134,38 @@ struct FAssetTableDependencySizes
 	int64 SharedDependenciesSize = 0;
 };
 
-struct FAssetTablePluginInfo
+class FAssetTablePluginInfo
 {
-	TSet<int32> PluginDependencies;
-	TSet<int32> DiscoveredPluginDependencies;
+	friend class FAssetTable;
+	friend class SAssetTableTreeView;
+
+public:
+	const TCHAR* GetName() const { return PluginName; }
+
+	int32 GetNumDependencies() const { return PluginDependencies.Num(); }
+	const TArray<int32>& GetDependencies() const { return PluginDependencies; }
+	int32 GetNumReferencers() const { return PluginReferencers.Num(); }
+	const TArray<int32>& GetReferencers() const { return PluginReferencers; }
+	int64 GetSize() const { return Size; }
+
+	int64 GetOrComputeTotalSizeInclusiveOfDependencies(const FAssetTable& OwningTable) const;
+
+	// This is the total size of all dependencies which are not used by any plugin that is not this plugin or one of its other dependencies
+	int64 GetOrComputeTotalSizeUniqueDependencies(const FAssetTable& OwningTable) const;
+	// This is the total size of all dependencies which are used by some plugin not referenced (directly or indirectly) by this plugin
+	int64 GetOrComputeTotalSizeSharedDependencies(const FAssetTable& OwningTable) const;
+
+private:
+	void ComputeDependencySizes(const FAssetTable& OwningTable) const;
+
+	TArray<int32> PluginDependencies;
+	TArray<int32> DiscoveredPluginDependencies;
+	TArray<int32> PluginReferencers;
 	const TCHAR* PluginName = nullptr;
-	int64 Size = -1;
+	int64 Size = 0;
+	mutable int64 InclusiveSize = -1;
+	mutable int64 UniqueDependenciesSize = -1;
+	mutable int64 SharedDependenciesSize = -1;
 };
 
 class FAssetTableRow
@@ -237,6 +264,7 @@ public:
 	FAssetTablePluginInfo& GetOrCreatePluginInfo(const TCHAR* StoredPluginName);
 
 	const FAssetTablePluginInfo& GetPluginInfoByIndex(int32 PluginIndex) const { return Plugins[PluginIndex]; }
+	FAssetTablePluginInfo& GetPluginInfoByIndex(int32 PluginIndex) { return Plugins[PluginIndex]; }
 	const FAssetTablePluginInfo& GetPluginInfoByIndexChecked(int32 PluginIndex) const { check(IsValidPluginIndex(PluginIndex)); return Plugins[PluginIndex]; }
 	int32 GetIndexForPlugin(const TCHAR* StoredPluginName) const;
 	int32 GetNumPlugins() const { return Plugins.Num(); }
diff --git a/Engine/Plugins/Editor/AssetManagerEditor/Source/AssetManagerEditor/Private/TreeView/AssetTreeNode.cpp b/Engine/Plugins/Editor/AssetManagerEditor/Source/AssetManagerEditor/Private/TreeView/AssetTreeNode.cpp
--- a/Engine/Plugins/Editor/AssetManagerEditor/Source/AssetManagerEditor/Private/TreeView/AssetTreeNode.cpp
+++ b/Engine/Plugins/Editor/AssetManagerEditor/Source/AssetManagerEditor/Private/TreeView/AssetTreeNode.cpp
@@ -97,7 +97,7 @@ FLinearColor FAssetTreeNode::GetColor(EStyle Style) const
 
 FAssetTreeNode::EStyle FAssetTreeNode::GetStyle() const
 {
-	return IsValidAsset() ? EStyle::EAsset : ( IsGroup() ? EStyle::EGroup : EStyle::EDefault );
+	return IsGroup() ? EStyle::EGroup : (IsValidAsset() ? EStyle::EAsset : EStyle::EDefault);
 }
 
 const FSlateBrush* FAssetTreeNode::GetIcon() const
@@ -118,7 +118,7 @@ FLinearColor FAssetTreeNode::GetColor() const
 
 FAssetTreeNode::EStyle FAssetDependenciesGroupTreeNode::GetStyle() const
 {
-	return EStyle::EGroup;
+	return EStyle::EDependencies;
 }
 
 ////////////////////////////////////////////////////////////////////////////////////////////////////
@@ -282,14 +282,6 @@ bool FPluginDependenciesGroupNode::OnLazyCreateChildren(TSharedPtr<class UE::Ins
 	// |       +-- [asset:{Asset1b}]
 	// |       ...
 	// |
-	// +-- [group:{DependentPlugin2}] (self, no further dependencies) // FPluginSimpleGroupNode
-	// |   |
-	// |   +-- [asset:{Asset2a}]
-	// |   |
-	// |   +-- [asset:{Asset2b}]
-	// |   ...
-	// |
-	// ...
 
 	FAssetTable& AssetTable = GetAssetTableChecked();
 	if (AssetTable.IsValidPluginIndex(PluginIndex))
@@ -297,24 +289,15 @@ bool FPluginDependenciesGroupNode::OnLazyCreateChildren(TSharedPtr<class UE::Ins
 		const FAssetTablePluginInfo& PluginInfo = AssetTable.GetPluginInfoByIndex(PluginIndex);
 
 		// Add dependent plugins.
-		for (int32 DependentPluginIndex : PluginInfo.PluginDependencies)
+		for (int32 DependentPluginIndex : PluginInfo.GetDependencies())
 		{
 			if (AssetTable.IsValidPluginIndex(DependentPluginIndex))
 			{
 				FName PluginGroupName = AssetTable.GetNameForPlugin(DependentPluginIndex);
 				const FAssetTablePluginInfo& DependentPluginInfo = AssetTable.GetPluginInfoByIndex(DependentPluginIndex);
-				if (DependentPluginInfo.PluginDependencies.Num() > 0)
-				{
-					TSharedPtr<FPluginAndDependenciesGroupNode> PluginGroup = MakeShared<FPluginAndDependenciesGroupNode>(PluginGroupName, GetAssetTableWeak(), DependentPluginIndex);
-					PluginGroup->CreateChildren()->AddAssetChildrenNodes();
-					AddChildAndSetParent(PluginGroup);
-				}
-				else
-				{
-					TSharedPtr<FPluginSimpleGroupNode> PluginGroup = MakeShared<FPluginSimpleGroupNode>(PluginGroupName, GetAssetTableWeak(), DependentPluginIndex);
-					PluginGroup->AddAssetChildrenNodes();
-					AddChildAndSetParent(PluginGroup);
-				}
+				TSharedPtr<FPluginAndDependenciesGroupNode> PluginGroup = MakeShared<FPluginAndDependenciesGroupNode>(PluginGroupName, GetAssetTableWeak(), DependentPluginIndex);
+				PluginGroup->CreateChildren()->AddAssetChildrenNodes();
+				AddChildAndSetParent(PluginGroup);
 			}
 		}
 	}
diff --git a/Engine/Plugins/Editor/AssetManagerEditor/Source/AssetManagerEditor/Private/TreeView/AssetTreeNode.h b/Engine/Plugins/Editor/AssetManagerEditor/Source/AssetManagerEditor/Private/TreeView/AssetTreeNode.h
--- a/Engine/Plugins/Editor/AssetManagerEditor/Source/AssetManagerEditor/Private/TreeView/AssetTreeNode.h
+++ b/Engine/Plugins/Editor/AssetManagerEditor/Source/AssetManagerEditor/Private/TreeView/AssetTreeNode.h
@@ -128,6 +128,8 @@ public:
 
 	void AddAssetChildrenNodes();
 
+	int32 GetPluginIndex() const { return PluginIndex; }
+
 protected:
 	int32 PluginIndex;
 };
diff --git a/Engine/Plugins/Editor/AssetManagerEditor/Source/AssetManagerEditor/Private/TreeView/SAssetTableTreeView.cpp b/Engine/Plugins/Editor/AssetManagerEditor/Source/AssetManagerEditor/Private/TreeView/SAssetTableTreeView.cpp
--- a/Engine/Plugins/Editor/AssetManagerEditor/Source/AssetManagerEditor/Private/TreeView/SAssetTableTreeView.cpp
+++ b/Engine/Plugins/Editor/AssetManagerEditor/Source/AssetManagerEditor/Private/TreeView/SAssetTableTreeView.cpp
@@ -98,7 +98,7 @@
 			bNeedsToOpenRegistry = false;
 			OpenRegistry();
 		}
-		if (bNeedsToRefreshAssets && !AssetRegistry->IsLoadingAssets())
+		if (bNeedsToRefreshAssets)
 		{
 			bNeedsToRefreshAssets = false;
 			RefreshAssets();
@@ -242,19 +242,19 @@
 
 void SAssetTableTreeView::InitAvailableViewPresets()
 {
-	////////////////////////////////////////////////////
-	//// Asset Type Breakdown View
+	//////////////////////////////////////////////////
+	// Plugin Dependency View
 
-	class FAssetTypeViewPreset : public UE::Insights::ITableTreeViewPreset
+	class FPluginDependencyView : public UE::Insights::ITableTreeViewPreset
 	{
 	public:
 		virtual FText GetName() const override
 		{
-			return LOCTEXT("AssetType_PresetName", "Group By Asset Type");
+			return LOCTEXT("PluginDepView_PresetName", "Plugin Dependency Analysis");
 		}
 		virtual FText GetToolTip() const override
 		{
-			return LOCTEXT("AssetType_PresetToolTip", "Asset Type Breakdown View\nConfigure the tree view to show a breakdown of assets by their asset type.");
+			return LOCTEXT("PluginDepView_PresetToolTip", "Plugin Dependency Analysis View\nConfigure the tree view to show a breakdown of assets by Plugin, showing also the dependencies between plugins.");
 		}
 		virtual FName GetSortColumn() const override
 		{
@@ -271,71 +271,71 @@
 			check(InAvailableGroupings[0]->Is<UE::Insights::FTreeNodeGroupingFlat>());
 			InOutCurrentGroupings.Add(InAvailableGroupings[0]);
 
-			const TSharedPtr<UE::Insights::FTreeNodeGrouping>* PrimaryTypeGrouping = InAvailableGroupings.FindByPredicate(
+			const TSharedPtr<UE::Insights::FTreeNodeGrouping>* PrimaryGrouping = InAvailableGroupings.FindByPredicate(
+				[](TSharedPtr<UE::Insights::FTreeNodeGrouping>& Grouping)
+				{
+					return Grouping->Is<FPluginDependencyGrouping>();
+				});
+			if (PrimaryGrouping)
+			{
+				InOutCurrentGroupings.Add(*PrimaryGrouping);
+			}
+
+			const TSharedPtr<UE::Insights::FTreeNodeGrouping>* SecondaryGrouping = InAvailableGroupings.FindByPredicate(
 				[](TSharedPtr<UE::Insights::FTreeNodeGrouping>& Grouping)
 				{
 					return Grouping->Is<UE::Insights::FTreeNodeGroupingByUniqueValueCString>() &&
 						Grouping->As<UE::Insights::FTreeNodeGroupingByUniqueValueCString>().GetColumnId() == FAssetTableColumns::TypeColumnId;
 				});
-			if (PrimaryTypeGrouping)
+			if (SecondaryGrouping)
 			{
-				InOutCurrentGroupings.Add(*PrimaryTypeGrouping);
-			}
-
-			const TSharedPtr<UE::Insights::FTreeNodeGrouping>* SecondaryTypeGrouping = InAvailableGroupings.FindByPredicate(
-				[](TSharedPtr<UE::Insights::FTreeNodeGrouping>& Grouping)
-				{
-					return Grouping->Is<UE::Insights::FTreeNodeGroupingByUniqueValueCString>() &&
-						Grouping->As<UE::Insights::FTreeNodeGroupingByUniqueValueCString>().GetColumnId() == FAssetTableColumns::PluginNameColumnId;
-				});
-			if (SecondaryTypeGrouping)
-			{
-				InOutCurrentGroupings.Add(*SecondaryTypeGrouping);
+				InOutCurrentGroupings.Add(*SecondaryGrouping);
 			}
 		}
 		virtual void GetColumnConfigSet(TArray<UE::Insights::FTableColumnConfig>& InOutConfigSet) const override
 		{
+
 			InOutConfigSet.Add({ UE::Insights::FTable::GetHierarchyColumnId(),              true, 400.0f });
 			InOutConfigSet.Add({ FAssetTableColumns::CountColumnId,                         true, 100.0f });
 			InOutConfigSet.Add({ FAssetTableColumns::StagedCompressedSizeColumnId,          true, 100.0f });
+			InOutConfigSet.Add({ FAssetTableColumns::PluginInclusiveSizeColumnId,	        true, 100.0f });
 			InOutConfigSet.Add({ FAssetTableColumns::TypeColumnId,                          true, 200.0f });
 			InOutConfigSet.Add({ FAssetTableColumns::NameColumnId,                          true, 200.0f });
 			InOutConfigSet.Add({ FAssetTableColumns::PathColumnId,                         !true, 400.0f });
-			InOutConfigSet.Add({ FAssetTableColumns::PrimaryTypeColumnId,                   true, 200.0f });
-			InOutConfigSet.Add({ FAssetTableColumns::PrimaryNameColumnId,                   true, 200.0f });
-			InOutConfigSet.Add({ FAssetTableColumns::TotalSizeUniqueDependenciesColumnId,   true, 100.0f });
-			InOutConfigSet.Add({ FAssetTableColumns::TotalSizeSharedDependenciesColumnId,   true, 100.0f });
-			InOutConfigSet.Add({ FAssetTableColumns::TotalSizeExternalDependenciesColumnId, true, 100.0f });
-			InOutConfigSet.Add({ FAssetTableColumns::TotalUsageCountColumnId,               true, 100.0f });
+			InOutConfigSet.Add({ FAssetTableColumns::PrimaryTypeColumnId,                  !true, 200.0f });
+			InOutConfigSet.Add({ FAssetTableColumns::PrimaryNameColumnId,                  !true, 200.0f });
+			InOutConfigSet.Add({ FAssetTableColumns::TotalSizeUniqueDependenciesColumnId,  !true, 100.0f });
+			InOutConfigSet.Add({ FAssetTableColumns::TotalSizeSharedDependenciesColumnId,  !true, 100.0f });
+			InOutConfigSet.Add({ FAssetTableColumns::TotalSizeExternalDependenciesColumnId,!true, 100.0f });
+			InOutConfigSet.Add({ FAssetTableColumns::TotalUsageCountColumnId,              !true, 100.0f });
 			InOutConfigSet.Add({ FAssetTableColumns::ChunksColumnId,                       !true, 200.0f });
-			InOutConfigSet.Add({ FAssetTableColumns::NativeClassColumnId,                   true, 200.0f });
+			InOutConfigSet.Add({ FAssetTableColumns::NativeClassColumnId,                  !true, 200.0f });
 			InOutConfigSet.Add({ FAssetTableColumns::PluginNameColumnId,                    true, 200.0f });
 		}
 	};
-	AvailableViewPresets.Add(MakeShared<FAssetTypeViewPreset>());
 
 	//////////////////////////////////////////////////
-	// Plugin Group View
-
-	class FPluginView : public UE::Insights::ITableTreeViewPreset
+	// Plugin, Type, Dependency View
+	
+	class FPluginTypeDependencyView : public UE::Insights::ITableTreeViewPreset
 	{
 	public:
 		virtual FText GetName() const override
 		{
-			return LOCTEXT("PluginGroupView_PresetName", "Group By Plugin");
+			return LOCTEXT("PluginTypeDepView_PresetName", "Asset Dependency Analysis");
 		}
 
 		virtual FText GetToolTip() const override
 		{
-			return LOCTEXT("PluginGroupView_PresetToolTip", "Group By Plugin\nGroup assets by plugin.");
+			return LOCTEXT("PluginTypeDepView_PresetToolTip", "Asset Dependency Analysis View\nConfigure the tree view to show a breakdown of assets by Plugin, Type, and Dependencies.");
 		}
 		virtual FName GetSortColumn() const override
 		{
-			return FAssetTableColumns::StagedCompressedSizeColumnId;
+			return UE::Insights::FTable::GetHierarchyColumnId();
 		}
 		virtual EColumnSortMode::Type GetSortMode() const override
 		{
-			return EColumnSortMode::Type::Descending;
+			return EColumnSortMode::Type::Ascending;
 		}
 		virtual void SetCurrentGroupings(const TArray<TSharedPtr<UE::Insights::FTreeNodeGrouping>>& InAvailableGroupings, TArray<TSharedPtr<UE::Insights::FTreeNodeGrouping>>& InOutCurrentGroupings) const override
 		{
@@ -344,63 +344,72 @@
 			check(InAvailableGroupings[0]->Is<UE::Insights::FTreeNodeGroupingFlat>());
 			InOutCurrentGroupings.Add(InAvailableGroupings[0]);
 
-			const TSharedPtr<UE::Insights::FTreeNodeGrouping>* PrimaryGrouping = InAvailableGroupings.FindByPredicate(
+			const TSharedPtr<UE::Insights::FTreeNodeGrouping>* GameFeaturePluginGrouping = InAvailableGroupings.FindByPredicate(
 				[](TSharedPtr<UE::Insights::FTreeNodeGrouping>& Grouping)
 				{
 					return Grouping->Is<UE::Insights::FTreeNodeGroupingByUniqueValueCString>() &&
 						Grouping->As<UE::Insights::FTreeNodeGroupingByUniqueValueCString>().GetColumnId() == FAssetTableColumns::PluginNameColumnId;
 				});
-			if (PrimaryGrouping)
+			if (GameFeaturePluginGrouping)
 			{
-				InOutCurrentGroupings.Add(*PrimaryGrouping);
+				InOutCurrentGroupings.Add(*GameFeaturePluginGrouping);
 			}
 
-			const TSharedPtr<UE::Insights::FTreeNodeGrouping>* SecondaryGrouping = InAvailableGroupings.FindByPredicate(
+			const TSharedPtr<UE::Insights::FTreeNodeGrouping>* PrimaryTypeGrouping = InAvailableGroupings.FindByPredicate(
 				[](TSharedPtr<UE::Insights::FTreeNodeGrouping>& Grouping)
 				{
 					return Grouping->Is<UE::Insights::FTreeNodeGroupingByUniqueValueCString>() &&
 						Grouping->As<UE::Insights::FTreeNodeGroupingByUniqueValueCString>().GetColumnId() == FAssetTableColumns::TypeColumnId;
 				});
-			if (SecondaryGrouping)
+			if (PrimaryTypeGrouping)
 			{
-				InOutCurrentGroupings.Add(*SecondaryGrouping);
+				InOutCurrentGroupings.Add(*PrimaryTypeGrouping);
+			}
+
+			const TSharedPtr<UE::Insights::FTreeNodeGrouping>* DependencyGrouping = InAvailableGroupings.FindByPredicate(
+				[](TSharedPtr<UE::Insights::FTreeNodeGrouping>& Grouping)
+				{
+					return Grouping->Is<FAssetDependencyGrouping>();
+				});
+			if (DependencyGrouping)
+			{
+				InOutCurrentGroupings.Add(*DependencyGrouping);
 			}
 		}
 		virtual void GetColumnConfigSet(TArray<UE::Insights::FTableColumnConfig>& InOutConfigSet) const override
 		{
 			InOutConfigSet.Add({ UE::Insights::FTable::GetHierarchyColumnId(),              true, 400.0f });
 			InOutConfigSet.Add({ FAssetTableColumns::CountColumnId,                         true, 100.0f });
-			InOutConfigSet.Add({ FAssetTableColumns::StagedCompressedSizeColumnId,          true, 100.0f });
-			InOutConfigSet.Add({ FAssetTableColumns::TypeColumnId,                          true, 200.0f });
-			InOutConfigSet.Add({ FAssetTableColumns::NameColumnId,                          true, 200.0f });
+			InOutConfigSet.Add({ FAssetTableColumns::TypeColumnId,                         !true, 200.0f });
+			InOutConfigSet.Add({ FAssetTableColumns::NameColumnId,                         !true, 200.0f });
 			InOutConfigSet.Add({ FAssetTableColumns::PathColumnId,                         !true, 400.0f });
 			InOutConfigSet.Add({ FAssetTableColumns::PrimaryTypeColumnId,                   true, 200.0f });
-			InOutConfigSet.Add({ FAssetTableColumns::PrimaryNameColumnId,                   true, 200.0f });
+			InOutConfigSet.Add({ FAssetTableColumns::PrimaryNameColumnId,                  !true, 200.0f });
+			InOutConfigSet.Add({ FAssetTableColumns::StagedCompressedSizeColumnId,          true, 100.0f });
 			InOutConfigSet.Add({ FAssetTableColumns::TotalSizeUniqueDependenciesColumnId,   true, 100.0f });
 			InOutConfigSet.Add({ FAssetTableColumns::TotalSizeSharedDependenciesColumnId,   true, 100.0f });
-			InOutConfigSet.Add({ FAssetTableColumns::TotalSizeExternalDependenciesColumnId, true, 100.0f });
-			InOutConfigSet.Add({ FAssetTableColumns::TotalUsageCountColumnId,               true, 100.0f });
+			InOutConfigSet.Add({ FAssetTableColumns::TotalSizeExternalDependenciesColumnId,!true, 100.0f });
+			InOutConfigSet.Add({ FAssetTableColumns::TotalUsageCountColumnId,              !true, 100.0f });
 			InOutConfigSet.Add({ FAssetTableColumns::ChunksColumnId,                       !true, 200.0f });
 			InOutConfigSet.Add({ FAssetTableColumns::NativeClassColumnId,                   true, 200.0f });
 			InOutConfigSet.Add({ FAssetTableColumns::PluginNameColumnId,                    true, 200.0f });
 		}
 	};
-	AvailableViewPresets.Add(MakeShared<FPluginView>());
+	 
+	////////////////////////////////////////////////////
+	//// Asset Type Breakdown View
 
-	//////////////////////////////////////////////////
-	// Plugin Dependency View
-
-	class FPluginDependencyView : public UE::Insights::ITableTreeViewPreset
+	class FAssetTypeViewPreset : public UE::Insights::ITableTreeViewPreset
 	{
 	public:
 		virtual FText GetName() const override
 		{
-			return LOCTEXT("PluginDepView_PresetName", "Plugin Dependency Analysis");
+			return LOCTEXT("AssetType_PresetName", "Group By Asset Type");
 		}
 
 		virtual FText GetToolTip() const override
 		{
-			return LOCTEXT("PluginDepView_PresetToolTip", "Plugin Dependency Analysis View\nConfigure the tree view to show a breakdown of assets by Game Feature Plugin, showing also the dependencies between plugins.");
+			return LOCTEXT("AssetType_PresetToolTip", "Asset Type Breakdown View\nConfigure the tree view to show a breakdown of assets by their asset type.");
 		}
 		virtual FName GetSortColumn() const override
 		{
@@ -417,25 +426,26 @@
 			check(InAvailableGroupings[0]->Is<UE::Insights::FTreeNodeGroupingFlat>());
 			InOutCurrentGroupings.Add(InAvailableGroupings[0]);
 
-			const TSharedPtr<UE::Insights::FTreeNodeGrouping>* PrimaryGrouping = InAvailableGroupings.FindByPredicate(
-				[](TSharedPtr<UE::Insights::FTreeNodeGrouping>& Grouping)
-				{
-					return Grouping->Is<FPluginDependencyGrouping>();
-				});
-			if (PrimaryGrouping)
-			{
-				InOutCurrentGroupings.Add(*PrimaryGrouping);
-			}
-
-			const TSharedPtr<UE::Insights::FTreeNodeGrouping>* SecondaryGrouping = InAvailableGroupings.FindByPredicate(
+			const TSharedPtr<UE::Insights::FTreeNodeGrouping>* TypeGrouping = InAvailableGroupings.FindByPredicate(
 				[](TSharedPtr<UE::Insights::FTreeNodeGrouping>& Grouping)
 				{
 					return Grouping->Is<UE::Insights::FTreeNodeGroupingByUniqueValueCString>() &&
 						Grouping->As<UE::Insights::FTreeNodeGroupingByUniqueValueCString>().GetColumnId() == FAssetTableColumns::TypeColumnId;
 				});
-			if (SecondaryGrouping)
+			if (TypeGrouping)
 			{
-				InOutCurrentGroupings.Add(*SecondaryGrouping);
+				InOutCurrentGroupings.Add(*TypeGrouping);
 			}
+
+			const TSharedPtr<UE::Insights::FTreeNodeGrouping>* PluginNameGrouping = InAvailableGroupings.FindByPredicate(
+				[](TSharedPtr<UE::Insights::FTreeNodeGrouping>& Grouping)
+				{
+					return Grouping->Is<UE::Insights::FTreeNodeGroupingByUniqueValueCString>() &&
+						Grouping->As<UE::Insights::FTreeNodeGroupingByUniqueValueCString>().GetColumnId() == FAssetTableColumns::PluginNameColumnId;
+				});
+			if (PluginNameGrouping)
+			{
+				InOutCurrentGroupings.Add(*PluginNameGrouping);
+			}
 		}
 		virtual void GetColumnConfigSet(TArray<UE::Insights::FTableColumnConfig>& InOutConfigSet) const override
@@ -451,13 +461,12 @@
 			InOutConfigSet.Add({ FAssetTableColumns::TotalSizeUniqueDependenciesColumnId,   true, 100.0f });
 			InOutConfigSet.Add({ FAssetTableColumns::TotalSizeSharedDependenciesColumnId,   true, 100.0f });
 			InOutConfigSet.Add({ FAssetTableColumns::TotalSizeExternalDependenciesColumnId, true, 100.0f });
-			InOutConfigSet.Add({ FAssetTableColumns::TotalUsageCountColumnId,               true, 100.0f });
+			InOutConfigSet.Add({ FAssetTableColumns::TotalUsageCountColumnId,              !true, 100.0f });
 			InOutConfigSet.Add({ FAssetTableColumns::ChunksColumnId,                       !true, 200.0f });
 			InOutConfigSet.Add({ FAssetTableColumns::NativeClassColumnId,                   true, 200.0f });
 			InOutConfigSet.Add({ FAssetTableColumns::PluginNameColumnId,                    true, 200.0f });
 		}
 	};
-	AvailableViewPresets.Add(MakeShared<FPluginDependencyView>());
 
 	////////////////////////////////////////////////////
 	//// Primary Type Breakdown View
@@ -523,13 +532,12 @@
 			InOutConfigSet.Add({ FAssetTableColumns::TotalSizeUniqueDependenciesColumnId,   true, 100.0f });
 			InOutConfigSet.Add({ FAssetTableColumns::TotalSizeSharedDependenciesColumnId,   true, 100.0f });
 			InOutConfigSet.Add({ FAssetTableColumns::TotalSizeExternalDependenciesColumnId, true, 100.0f });
-			InOutConfigSet.Add({ FAssetTableColumns::TotalUsageCountColumnId,               true, 100.0f });
+			InOutConfigSet.Add({ FAssetTableColumns::TotalUsageCountColumnId,              !true, 100.0f });
 			InOutConfigSet.Add({ FAssetTableColumns::ChunksColumnId,                       !true, 200.0f });
 			InOutConfigSet.Add({ FAssetTableColumns::NativeClassColumnId,                   true, 200.0f });
 			InOutConfigSet.Add({ FAssetTableColumns::PluginNameColumnId,                    true, 200.0f });
 		}
 	};
-	AvailableViewPresets.Add(MakeShared<FPrimaryTypeViewPreset>());
 
 
 	////////////////////////////////////////////////////
@@ -596,15 +604,98 @@
 			InOutConfigSet.Add({ FAssetTableColumns::TotalSizeUniqueDependenciesColumnId,   true, 100.0f });
 			InOutConfigSet.Add({ FAssetTableColumns::TotalSizeSharedDependenciesColumnId,   true, 100.0f });
 			InOutConfigSet.Add({ FAssetTableColumns::TotalSizeExternalDependenciesColumnId, true, 100.0f });
-			InOutConfigSet.Add({ FAssetTableColumns::TotalUsageCountColumnId,               true, 100.0f });
+			InOutConfigSet.Add({ FAssetTableColumns::TotalUsageCountColumnId,              !true, 100.0f });
 			InOutConfigSet.Add({ FAssetTableColumns::ChunksColumnId,                       !true, 200.0f });
 			InOutConfigSet.Add({ FAssetTableColumns::NativeClassColumnId,                   true, 200.0f });
 			InOutConfigSet.Add({ FAssetTableColumns::PluginNameColumnId,                    true, 200.0f });
 		}
 	};
+
+
+	//////////////////////////////////////////////////
+	// Plugin Group View
+
+	class FPluginView : public UE::Insights::ITableTreeViewPreset
+	{
+	public:
+		virtual FText GetName() const override
+		{
+			return LOCTEXT("PluginGroupView_PresetName", "Group By Plugin");
+		}
+
+		virtual FText GetToolTip() const override
+		{
+			return LOCTEXT("PluginGroupView_PresetToolTip", "Group By Plugin\nGroup assets by plugin.");
+		}
+		virtual FName GetSortColumn() const override
+		{
+			return FAssetTableColumns::StagedCompressedSizeColumnId;
+		}
+		virtual EColumnSortMode::Type GetSortMode() const override
+		{
+			return EColumnSortMode::Type::Descending;
+		}
+		virtual void SetCurrentGroupings(const TArray<TSharedPtr<UE::Insights::FTreeNodeGrouping>>& InAvailableGroupings, TArray<TSharedPtr<UE::Insights::FTreeNodeGrouping>>& InOutCurrentGroupings) const override
+		{
+			InOutCurrentGroupings.Reset();
+
+			check(InAvailableGroupings[0]->Is<UE::Insights::FTreeNodeGroupingFlat>());
+			InOutCurrentGroupings.Add(InAvailableGroupings[0]);
+
+			const TSharedPtr<UE::Insights::FTreeNodeGrouping>* PrimaryGrouping = InAvailableGroupings.FindByPredicate(
+				[](TSharedPtr<UE::Insights::FTreeNodeGrouping>& Grouping)
+				{
+					return Grouping->Is<UE::Insights::FTreeNodeGroupingByUniqueValueCString>() &&
+						Grouping->As<UE::Insights::FTreeNodeGroupingByUniqueValueCString>().GetColumnId() == FAssetTableColumns::PluginNameColumnId;
+				});
+			if (PrimaryGrouping)
+			{
+				InOutCurrentGroupings.Add(*PrimaryGrouping);
+			}
+
+			const TSharedPtr<UE::Insights::FTreeNodeGrouping>* SecondaryGrouping = InAvailableGroupings.FindByPredicate(
+				[](TSharedPtr<UE::Insights::FTreeNodeGrouping>& Grouping)
+				{
+					return Grouping->Is<UE::Insights::FTreeNodeGroupingByUniqueValueCString>() &&
+						Grouping->As<UE::Insights::FTreeNodeGroupingByUniqueValueCString>().GetColumnId() == FAssetTableColumns::TypeColumnId;
+				});
+			if (SecondaryGrouping)
+			{
+				InOutCurrentGroupings.Add(*SecondaryGrouping);
+			}
+		}
+		virtual void GetColumnConfigSet(TArray<UE::Insights::FTableColumnConfig>& InOutConfigSet) const override
+		{
+			InOutConfigSet.Add({ UE::Insights::FTable::GetHierarchyColumnId(),              true, 400.0f });
+			InOutConfigSet.Add({ FAssetTableColumns::CountColumnId,                         true, 100.0f });
+			InOutConfigSet.Add({ FAssetTableColumns::StagedCompressedSizeColumnId,          true, 100.0f });
+			InOutConfigSet.Add({ FAssetTableColumns::PluginInclusiveSizeColumnId,	        true, 100.0f });
+			InOutConfigSet.Add({ FAssetTableColumns::TypeColumnId,                          true, 200.0f });
+			InOutConfigSet.Add({ FAssetTableColumns::NameColumnId,                          true, 200.0f });
+			InOutConfigSet.Add({ FAssetTableColumns::PathColumnId,                         !true, 400.0f });
+			InOutConfigSet.Add({ FAssetTableColumns::PrimaryTypeColumnId,                   true, 200.0f });
+			InOutConfigSet.Add({ FAssetTableColumns::PrimaryNameColumnId,                   true, 200.0f });
+			InOutConfigSet.Add({ FAssetTableColumns::TotalSizeUniqueDependenciesColumnId,   true, 100.0f });
+			InOutConfigSet.Add({ FAssetTableColumns::TotalSizeSharedDependenciesColumnId,   true, 100.0f });
+			InOutConfigSet.Add({ FAssetTableColumns::TotalSizeExternalDependenciesColumnId, true, 100.0f });
+			InOutConfigSet.Add({ FAssetTableColumns::TotalUsageCountColumnId,              !true, 100.0f });
+			InOutConfigSet.Add({ FAssetTableColumns::ChunksColumnId,                       !true, 200.0f });
+			InOutConfigSet.Add({ FAssetTableColumns::NativeClassColumnId,                   true, 200.0f });
+			InOutConfigSet.Add({ FAssetTableColumns::PluginNameColumnId,                    true, 200.0f });
+		}
+	};
+
+	//////////////////////////////////////////////////
+
+	AvailableViewPresets.Add(MakeShared<FPluginDependencyView>());
+	AvailableViewPresets.Add(MakeShared<FPluginTypeDependencyView>());
+	AvailableViewPresets.Add(MakeShared<FAssetTypeViewPreset>());
+	AvailableViewPresets.Add(MakeShared<FPluginView>());
+	AvailableViewPresets.Add(MakeShared<FPrimaryTypeViewPreset>());
 	AvailableViewPresets.Add(MakeShared<FClassTypeViewPreset>());
 
-	//////////////////////////////////////////////////
+	/////////////////////////////////////////////////
+
 	SelectedViewPreset = AvailableViewPresets[0];
 }
 
@@ -1309,6 +1400,8 @@
 
 	UE_LOG(LogInsights, Log, TEXT("[AssetTree] Build asset table..."));
 
+	CancelCurrentAsyncOp();
+
 	UE::Insights::FStopwatch Stopwatch;
 	Stopwatch.Start();
 
@@ -1483,6 +1576,7 @@
 		{
 			TArray<FPluginReferenceDescriptor> PluginReferences;
 			IPluginManager::Get().GetPluginDependencies(PluginEntry.Key, PluginReferences);
+			const int32 PluginIndex = AssetTable->GetIndexForPlugin(StoredPluginName);
 			for (const FPluginReferenceDescriptor& ReferenceDescriptor : PluginReferences)
 			{
 				const TCHAR* StoredReferencePluginName = AssetTable->StoreStr(ReferenceDescriptor.Name);
@@ -1494,7 +1588,8 @@
 				}
 				// Note that we can't use PluginInfo directly because the above code could have grown the array 
 				// and that would invalidate the reference
-				AssetTable->GetOrCreatePluginInfo(StoredPluginName).PluginDependencies.Add(DependencyIndex);
+				AssetTable->GetOrCreatePluginInfo(StoredPluginName).PluginDependencies.AddUnique(DependencyIndex);
+				AssetTable->GetPluginInfoByIndex(DependencyIndex).PluginReferencers.AddUnique(AssetTable->GetIndexForPlugin(StoredPluginName));
 			}
 		}
 		else
@@ -1513,7 +1608,7 @@
 		{
 			int32 DependencyIndex = AssetTable->GetIndexForPlugin(DependencyName);
 			ensureAlways(DependencyIndex != -1); // These were created above.
-			PluginInfo.DiscoveredPluginDependencies.Add(DependencyIndex);
+			PluginInfo.DiscoveredPluginDependencies.AddUnique(DependencyIndex);
 		}
 	}
 
