PR #10517: Add support for setting AllowedTypes on FPrimaryAssetId parameters
To use this, add UPARAM(meta = (AllowedTypes = "PrimaryAssetType")) to a FPrimaryAssetId blueprint function parameter like how it already worked for variables

diff --git a/Engine/Plugins/Editor/AssetManagerEditor/Source/AssetManagerEditor/Private/PrimaryAssetIdCustomization.cpp b/Engine/Plugins/Editor/AssetManagerEditor/Source/AssetManagerEditor/Private/PrimaryAssetIdCustomization.cpp
--- a/Engine/Plugins/Editor/AssetManagerEditor/Source/AssetManagerEditor/Private/PrimaryAssetIdCustomization.cpp
+++ b/Engine/Plugins/Editor/AssetManagerEditor/Source/AssetManagerEditor/Private/PrimaryAssetIdCustomization.cpp
@@ -260,6 +260,23 @@ TSharedRef<SWidget>	SPrimaryAssetIdGraphPin::GetDefaultValueWidget()
 {
 	FString DefaultString = GraphPinObj->GetDefaultAsString();
 	CurrentId = FPrimaryAssetId(DefaultString);
+	
+	TArray<FPrimaryAssetType> AllowedTypes;
+	if (UEdGraphNode* OwningNode = GraphPinObj->GetOwningNode())
+	{
+		const FString TypeFilterString = OwningNode->GetPinMetaData(GraphPinObj->GetFName(), TEXT("AllowedTypes"));
+		if( !TypeFilterString.IsEmpty() )
+		{
+			TArray<FString> CustomTypeFilterNames;
+			TypeFilterString.ParseIntoArray(CustomTypeFilterNames, TEXT(","), true);
+
+			for(auto It = CustomTypeFilterNames.CreateConstIterator(); It; ++It)
+			{
+				const FString& TypeName = *It;
+				AllowedTypes.Add(*TypeName);
+			}
+		}
+	}
 
 	return SNew(SHorizontalBox)
 		.Visibility(this, &SGraphPin::GetDefaultValueVisibility)
@@ -269,7 +286,7 @@ TSharedRef<SWidget>	SPrimaryAssetIdGraphPin::GetDefaultValueWidget()
 			IAssetManagerEditorModule::MakePrimaryAssetIdSelector(
 				FOnGetPrimaryAssetDisplayText::CreateSP(this, &SPrimaryAssetIdGraphPin::GetDisplayText),
 				FOnSetPrimaryAssetId::CreateSP(this, &SPrimaryAssetIdGraphPin::OnIdSelected),
-				true)
+				true, AllowedTypes)
 		]
 		+ SHorizontalBox::Slot()
 		.AutoWidth()
