Add a plugin filter to the Reference Viewer

* The filter can be enabled or disabled using a checkbox.
* A ComboButton displays a dropdown list of plugins to filter by. Only plugins that actually contain assets visible in the unfiltered graph are shown in the list.
* The dropdown list shows non-plugins such as /Engine and /Game first after which all real plugins are listed.
* One downside of the new plugin filter is that nodes that come from plugins that do not match the filter will be hidden, including their children even if the children match the filter. So if we have the graph A -> B -> C and filter to show the plugin of C but hide the plugin of B then both B and C will be hidden.

diff --git a/Engine/Plugins/Editor/AssetManagerEditor/Source/AssetManagerEditor/AssetManagerEditor.Build.cs b/Engine/Plugins/Editor/AssetManagerEditor/Source/AssetManagerEditor/AssetManagerEditor.Build.cs
--- a/Engine/Plugins/Editor/AssetManagerEditor/Source/AssetManagerEditor/AssetManagerEditor.Build.cs
+++ b/Engine/Plugins/Editor/AssetManagerEditor/Source/AssetManagerEditor/AssetManagerEditor.Build.cs
@@ -25,8 +25,8 @@ public class AssetManagerEditor : ModuleRules
 			new string[] {
 				"Slate",
 				"SlateCore",
-                "ApplicationCore",
-                "InputCore",
+				"ApplicationCore",
+				"InputCore",
 				"EditorFramework",
 				"UnrealEd",
 				"AssetRegistry",
@@ -50,7 +50,7 @@ public class AssetManagerEditor : ModuleRules
 				"SourceControl",
 				"SourceControlWindows",
 				"UncontrolledChangelists",
-				"Projects"
+				"Projects",
 			}
 		);
 	}
diff --git a/Engine/Plugins/Editor/AssetManagerEditor/Source/AssetManagerEditor/Private/ReferenceViewer/EdGraph_ReferenceViewer.cpp b/Engine/Plugins/Editor/AssetManagerEditor/Source/AssetManagerEditor/Private/ReferenceViewer/EdGraph_ReferenceViewer.cpp
--- a/Engine/Plugins/Editor/AssetManagerEditor/Source/AssetManagerEditor/Private/ReferenceViewer/EdGraph_ReferenceViewer.cpp
+++ b/Engine/Plugins/Editor/AssetManagerEditor/Source/AssetManagerEditor/Private/ReferenceViewer/EdGraph_ReferenceViewer.cpp
@@ -13,6 +13,7 @@
 #include "CollectionManagerModule.h"
 #include "AssetManagerEditorModule.h"
 #include "Engine/AssetManager.h"
+#include "Interfaces/IPluginManager.h"
 
 #include UE_INLINE_GENERATED_CPP_BY_NAME(EdGraph_ReferenceViewer)
 
@@ -142,6 +143,22 @@ void UEdGraph_ReferenceViewer::SetCurrentCollectionFilter(FName NewFilter)
 	CurrentCollectionFilter = NewFilter;
 }
 
+TArray<FName> UEdGraph_ReferenceViewer::GetCurrentPluginFilter() const
+{
+	return CurrentPluginFilter;
+}
+
+void UEdGraph_ReferenceViewer::SetCurrentPluginFilter(TArray<FName> NewFilter)
+{
+	CurrentPluginFilter = NewFilter;
+}
+
+TArray<FName> UEdGraph_ReferenceViewer::GetEncounteredPluginsAmongNodes() const
+{
+	return EncounteredPluginsAmongNodes;
+}
+
+
 void UEdGraph_ReferenceViewer::SetCurrentFilterCollection(TSharedPtr< TFilterCollection< FReferenceNodeInfo& > > InFilterCollection )
 {
 	FilterCollection = InFilterCollection;
@@ -177,7 +194,6 @@ FAssetManagerDependencyQuery UEdGraph_ReferenceViewer::GetReferenceSearchFlags(b
 
 UEdGraphNode_Reference* UEdGraph_ReferenceViewer::ConstructNodes(const TArray<FAssetIdentifier>& GraphRootIdentifiers, const FIntPoint& GraphRootOrigin )
 {
-	UEdGraphNode_Reference* RootNode = nullptr;
 	if (GraphRootIdentifiers.Num() > 0)
 	{
 		// It both were false, nothing (other than the GraphRootIdentifiers) would be displayed
@@ -197,13 +213,30 @@ UEdGraphNode_Reference* UEdGraph_ReferenceViewer::ConstructNodes(const TArray<FA
 			}
 		}
 
+		// Prepare for plugin filtering.
+		if (ShouldFilterByPlugin())
+		{
+			// Collect plugin names from assets reachable in the graph if the graph had been unfiltered.
+			EncounteredPluginsAmongNodes.Empty();
+			GetUnfilteredGraphPluginNames(GraphRootIdentifiers, EncounteredPluginsAmongNodes);
+
+			// Remove plugins from the current filter that were not encountered in the new unfiltered graph.
+			for (TArray<FName>::TIterator It(CurrentPluginFilter); It; ++It)
+			{
+				if (!EncounteredPluginsAmongNodes.Contains(*It))
+				{
+					It.RemoveCurrent();
+				}
+			}
+		}
+
 		// Create & Populate the NodeInfo Maps 
 		// Note to add an empty parent to the root so that if the root node again gets found again as a duplicate, that next parent won't be 
 		// identified as the primary root and also it will appear as having multiple parents.
 		TMap<FAssetIdentifier, FReferenceNodeInfo> NewReferenceNodeInfos;
 		for (const FAssetIdentifier& RootIdentifier : GraphRootIdentifiers)
 		{
-			FReferenceNodeInfo& RootNodeInfo = NewReferenceNodeInfos.FindOrAdd( RootIdentifier, FReferenceNodeInfo(RootIdentifier, true));
+			FReferenceNodeInfo& RootNodeInfo = NewReferenceNodeInfos.FindOrAdd(RootIdentifier, FReferenceNodeInfo(RootIdentifier, true));
 			RootNodeInfo.Parents.Emplace(FAssetIdentifier(NAME_None));
 		}
 		if (!Settings->GetFindPathEnabled())
@@ -564,6 +597,11 @@ void UEdGraph_ReferenceViewer::GetSortedLinks(const TArray<FAssetIdentifier>& Id
 		{
 			It.RemoveCurrent();
 		}
+
+		else if (!IsPackageIdentifierPassingPluginFilter(It.Key()))
+		{
+			It.RemoveCurrent();
+		}
 	}
 }
 
@@ -581,6 +619,31 @@ bool UEdGraph_ReferenceViewer::IsPackageIdentifierPassingFilter(const FAssetIden
 	return true;
 }
 
+bool UEdGraph_ReferenceViewer::IsPackageIdentifierPassingPluginFilter(const FAssetIdentifier& InAssetIdentifier) const
+{
+	if (!ShouldFilterByPlugin())
+	{
+		return true;
+	}
+	
+	if (!InAssetIdentifier.IsPackage())
+	{
+		return true;
+	}
+	
+	const FString AssetPath = InAssetIdentifier.PackageName.ToString();
+
+	for (const FName& PluginName : CurrentPluginFilter)
+	{
+		if (AssetPath.StartsWith("/" + PluginName.ToString()))
+		{
+			return true;
+		}
+	}
+
+	return false;
+}
+
 bool UEdGraph_ReferenceViewer::IsAssetPassingSearchTextFilter(const FAssetIdentifier& InAssetIdentifier) const
 {
 	if (Settings->IsShowFilteredPackagesOnly() && IsAssetIdentifierPassingSearchFilterCallback.IsSet() && !(*IsAssetIdentifierPassingSearchFilterCallback)(InAssetIdentifier))
@@ -591,6 +654,81 @@ bool UEdGraph_ReferenceViewer::IsAssetPassingSearchTextFilter(const FAssetIdenti
 	return true;
 }
 
+void UEdGraph_ReferenceViewer::GetUnfilteredGraphPluginNamesRecursive(bool bReferencers, const FAssetIdentifier& InAssetIdentifier, int32 InCurrentDepth, int32 InMaxDepth, const FAssetManagerDependencyQuery& Query, TArray<FAssetIdentifier>& OutAssetIdentifiers)
+{
+	if (ExceedsMaxSearchDepth(InCurrentDepth, InMaxDepth))
+	{
+		return;
+	}
+	
+	IAssetRegistry& AssetRegistry = FModuleManager::LoadModuleChecked<FAssetRegistryModule>("AssetRegistry").Get();
+	TArray<FAssetDependency> LinksToAsset;
+	if (bReferencers)
+	{
+		AssetRegistry.GetReferencers(InAssetIdentifier, LinksToAsset);
+	}
+	else
+	{
+		AssetRegistry.GetDependencies(InAssetIdentifier, LinksToAsset);
+	}
+	
+	for (const FAssetDependency& Link : LinksToAsset)
+	{
+		// Don't add assets that will be hidden by Reference Viewer settings the user cannot change.
+		if (!IsPackageIdentifierPassingFilter(Link.AssetId))
+		{
+			continue;
+		}
+		
+		OutAssetIdentifiers.Add(Link.AssetId);
+
+		GetUnfilteredGraphPluginNamesRecursive(bReferencers, Link.AssetId, InCurrentDepth + 1, InMaxDepth, Query, OutAssetIdentifiers);
+	}
+}
+
+void UEdGraph_ReferenceViewer::GetUnfilteredGraphPluginNames(TArray<FAssetIdentifier> RootIdentifiers, TArray<FName>& OutPluginNames)
+{
+	TRACE_CPUPROFILER_EVENT_SCOPE(UEdGraph_ReferenceViewer::GetUnfilteredGraphPluginNames);
+	
+	const FAssetManagerDependencyQuery Query = GetReferenceSearchFlags(false);
+	
+	TArray<FAssetIdentifier> AssetIdentifiers;
+	for (const FAssetIdentifier& RootIdentifier : RootIdentifiers)
+	{
+		GetUnfilteredGraphPluginNamesRecursive(true, RootIdentifier, 0, Settings->GetSearchReferencerDepthLimit(), Query, AssetIdentifiers);
+		GetUnfilteredGraphPluginNamesRecursive(false, RootIdentifier, 0, Settings->GetSearchDependencyDepthLimit(), Query, AssetIdentifiers);
+	}
+
+	for (const FAssetIdentifier& AssetIdentifier : AssetIdentifiers)
+	{
+		if (!AssetIdentifier.IsPackage())
+		{
+			continue;
+		}
+
+		FString FirstPathSegment;
+		{
+			FString AssetPath = AssetIdentifier.PackageName.ToString();
+			
+			// Chop of any leading slashes.
+			while (AssetPath.StartsWith("/"))
+			{
+				AssetPath = AssetPath.Mid(1);
+			}
+
+			const int32 SecondSlash = AssetPath.Find("/");
+			if (SecondSlash != INDEX_NONE)
+			{
+				AssetPath = AssetPath.Left(SecondSlash);
+			}
+			
+			FirstPathSegment = AssetPath;
+		}
+
+		OutPluginNames.AddUnique(FName(FirstPathSegment));
+	}
+}
+
 void
 UEdGraph_ReferenceViewer::RecursivelyPopulateNodeInfos(bool bInReferencers, const TArray<FAssetIdentifier>& Identifiers, TMap<FAssetIdentifier, FReferenceNodeInfo>& InNodeInfos, int32 InCurrentDepth, int32 InMaxDepth)
 {
@@ -847,3 +985,8 @@ bool UEdGraph_ReferenceViewer::ShouldFilterByCollection() const
 	return Settings->GetEnableCollectionFilter() && CurrentCollectionFilter != NAME_None;
 }
 
+bool UEdGraph_ReferenceViewer::ShouldFilterByPlugin() const
+{
+	return Settings->GetEnablePluginFilter() && CurrentPluginFilter.Num() > 0;
+}
+
diff --git a/Engine/Plugins/Editor/AssetManagerEditor/Source/AssetManagerEditor/Private/ReferenceViewer/ReferenceViewerSettings.cpp b/Engine/Plugins/Editor/AssetManagerEditor/Source/AssetManagerEditor/Private/ReferenceViewer/ReferenceViewerSettings.cpp
--- a/Engine/Plugins/Editor/AssetManagerEditor/Source/AssetManagerEditor/Private/ReferenceViewer/ReferenceViewerSettings.cpp
+++ b/Engine/Plugins/Editor/AssetManagerEditor/Source/AssetManagerEditor/Private/ReferenceViewer/ReferenceViewerSettings.cpp
@@ -197,6 +197,17 @@ void UReferenceViewerSettings::SetEnableCollectionFilter(bool bEnabled)
 	SaveConfig();
 }
 
+bool UReferenceViewerSettings::GetEnablePluginFilter() const
+{
+	return bEnablePluginFilter;
+}
+
+void UReferenceViewerSettings::SetEnablePluginFilter(bool bEnabled)
+{
+	bEnablePluginFilter = bEnabled;
+	SaveConfig();
+}
+
 bool UReferenceViewerSettings::IsShowPath() const
 {
 	return bIsShowPath;
diff --git a/Engine/Plugins/Editor/AssetManagerEditor/Source/AssetManagerEditor/Private/ReferenceViewer/SReferenceViewer.cpp b/Engine/Plugins/Editor/AssetManagerEditor/Source/AssetManagerEditor/Private/ReferenceViewer/SReferenceViewer.cpp
--- a/Engine/Plugins/Editor/AssetManagerEditor/Source/AssetManagerEditor/Private/ReferenceViewer/SReferenceViewer.cpp
+++ b/Engine/Plugins/Editor/AssetManagerEditor/Source/AssetManagerEditor/Private/ReferenceViewer/SReferenceViewer.cpp
@@ -25,12 +25,15 @@
 #include "Toolkits/GlobalEditorCommonCommands.h"
 #include "Engine/AssetManager.h"
 #include "ReferenceViewer/SReferenceViewerFilterBar.h"
+#include "Widgets/Input/SCheckBox.h"
 #include "Widgets/Input/SComboBox.h"
+#include "Widgets/Input/SComboButton.h"
 #include "HAL/PlatformApplicationMisc.h"
 #include "AssetManagerEditorModule.h"
 
 #include "ObjectTools.h"
 #include "Selection.h"
+#include "Interfaces/IPluginManager.h"
 
 #define LOCTEXT_NAMESPACE "ReferenceViewer"
 
@@ -57,7 +60,7 @@ SReferenceViewer::~SReferenceViewer()
 		if ( ensure(GraphObj) )
 		{
 			GraphObj->RemoveFromRoot();
-		}		
+		}
 	}
 }
 
@@ -73,7 +76,7 @@ void SReferenceViewer::Construct(const FArguments& InArgs)
 	// Set up the history manager
 	HistoryManager.SetOnApplyHistoryData(FOnApplyHistoryData::CreateSP(this, &SReferenceViewer::OnApplyHistoryData));
 	HistoryManager.SetOnUpdateHistoryData(FOnUpdateHistoryData::CreateSP(this, &SReferenceViewer::OnUpdateHistoryData));
-
+	
 	// Create the graph
 	GraphObj = NewObject<UEdGraph_ReferenceViewer>();
 	GraphObj->Schema = UReferenceViewerSchema::StaticClass();
@@ -105,6 +108,7 @@ void SReferenceViewer::Construct(const FArguments& InArgs)
 	FixAndHideSearchDepthLimit = 0;
 	FixAndHideSearchBreadthLimit = 0;
 	bShowCollectionFilter = true;
+	bShowPluginFilter = true;
 	bShowShowReferencesOptions = true;
 	bShowShowSearchableNames = true;
 	bShowShowCodePackages = true;
@@ -469,6 +473,50 @@ void SReferenceViewer::Construct(const FArguments& InArgs)
 						]
 					]
 
+					+SVerticalBox::Slot()
+					.AutoHeight()
+					[
+						SNew(SHorizontalBox)
+						.Visibility_Lambda([this]() { return (bShowPluginFilter ? EVisibility::Visible : EVisibility::Collapsed); })
+
+						+SHorizontalBox::Slot()
+						.VAlign(VAlign_Center)
+						.FillWidth(1.0)
+						.Padding(2.f)
+						[
+							SNew(STextBlock)
+							.Text(LOCTEXT("PluginFilter", "Plugin Filter"))
+						]
+
+						+SHorizontalBox::Slot()
+						.AutoWidth()
+						.VAlign(VAlign_Center)
+						.Padding(2.f)
+						[
+							SAssignNew(PluginFilterCheckbox, SCheckBox)
+							.OnCheckStateChanged( this, &SReferenceViewer::OnEnablePluginFilterChanged )
+							.IsChecked( this, &SReferenceViewer::IsEnablePluginFilterChecked )
+						]
+
+						+SHorizontalBox::Slot()
+						.AutoWidth()
+						.VAlign(VAlign_Center)
+						.Padding(2.f)
+						[
+							SNew(SBox)
+							.WidthOverride(100)
+							[
+								SAssignNew(PluginsComboButton, SComboButton)
+								.OnGetMenuContent(this, &SReferenceViewer::BuildPluginFilterMenu)
+								.ButtonContent()
+								[
+									SNew(STextBlock)
+									.Text(this, &SReferenceViewer::GetPluginComboButtonText)
+									.ToolTipText(this, &SReferenceViewer::GetPluginComboButtonText)
+								]
+							]
+						]
+					]
 				]
 				] // SHorizontalBox::Slot()
 
@@ -546,6 +594,7 @@ void SReferenceViewer::SetGraphRootIdentifiers(const TArray<FAssetIdentifier>& N
 		Settings->SetSearchBreadthLimitEnabled(true);
 	}
 	bShowCollectionFilter = ReferenceViewerParams.bShowCollectionFilter;
+	bShowPluginFilter = ReferenceViewerParams.bShowPluginFilter;
 	bShowShowReferencesOptions = ReferenceViewerParams.bShowShowReferencesOptions;
 	bShowShowSearchableNames = ReferenceViewerParams.bShowShowSearchableNames;
 	bShowShowCodePackages = ReferenceViewerParams.bShowShowCodePackages;
@@ -1061,6 +1110,163 @@ FText SReferenceViewer::GetCollectionFilterText() const
 	return FText::FromName(GraphObj->GetCurrentCollectionFilter());
 }
 
+void SReferenceViewer::OnEnablePluginFilterChanged(ECheckBoxState NewState)
+{
+	const bool bNewValue = NewState == ECheckBoxState::Checked;
+
+	const bool bCurrentValue = Settings->GetEnablePluginFilter();
+	if (bCurrentValue != bNewValue)
+	{
+		Settings->SetEnablePluginFilter(NewState == ECheckBoxState::Checked);
+		RebuildGraph();
+	}
+}
+
+ECheckBoxState SReferenceViewer::IsEnablePluginFilterChecked() const
+{
+	return Settings->GetEnablePluginFilter() ? ECheckBoxState::Checked : ECheckBoxState::Unchecked;
+}
+
+void SReferenceViewer::PluginFilterAddMenuEntry(FMenuBuilder& MenuBuilder, const FName& PluginName)
+{
+	FExecuteAction ActionClicked = FExecuteAction::CreateLambda([this, PluginName]()
+	{
+		TArray<FName> CurrentPluginFilter = GraphObj->GetCurrentPluginFilter();
+		// We just got checked if we don't exist in the current plugin filter.
+		const bool bNewChecked = !CurrentPluginFilter.Contains(PluginName);
+
+		if (bNewChecked)
+		{
+			// Make sure plugin filtering is enabled now that something was checked.
+			Settings->SetEnablePluginFilter(true);
+
+			CurrentPluginFilter.AddUnique(PluginName);
+		}
+		else if (CurrentPluginFilter.Contains(PluginName))
+		{
+			CurrentPluginFilter.RemoveAll([PluginName](const FName& Name)
+			{
+				return Name == PluginName;
+			});
+		}
+
+		GraphObj->SetCurrentPluginFilter(CurrentPluginFilter);
+		RebuildGraph();
+	});
+
+	FIsActionChecked ActionChecked = FIsActionChecked::CreateLambda([this, PluginName]() -> bool
+	{
+		return GraphObj->GetCurrentPluginFilter().Contains(PluginName);
+	});
+
+	MenuBuilder.AddMenuEntry(
+		FText::FromName(PluginName),
+		FText::FromName(PluginName),
+		FSlateIcon(),
+		FUIAction(ActionClicked, FCanExecuteAction(), ActionChecked),
+		NAME_None, //InExtensionHook
+		EUserInterfaceActionType::ToggleButton
+		);
+}
+
+TSharedRef<SWidget> SReferenceViewer::BuildPluginFilterMenu()
+{
+	FMenuBuilder MenuBuilder(false, nullptr);
+
+	MenuBuilder.AddMenuEntry(
+		LOCTEXT("PluginFilterSelectAll", "Select All"),
+		LOCTEXT("PluginFilterSelectAll", "Select all plugins."),
+		FSlateIcon(),
+		FExecuteAction::CreateLambda([this]()
+		{
+			// Make sure plugin filtering is enabled.
+			Settings->SetEnablePluginFilter(true);
+
+			GraphObj->SetCurrentPluginFilter(GraphObj->GetEncounteredPluginsAmongNodes());
+			RebuildGraph();
+		})
+	);
+	
+	MenuBuilder.AddMenuEntry(
+		LOCTEXT("PluginFilterSelectNone", "Select None"),
+		LOCTEXT("PluginFilterSelectNone", "Select no plugins."),
+		FSlateIcon(),
+		FExecuteAction::CreateLambda([this]()
+		{
+			// Make sure plugin filtering is enabled.
+			Settings->SetEnablePluginFilter(true);
+
+			GraphObj->SetCurrentPluginFilter(TArray<FName>());
+			RebuildGraph();
+		})
+	);
+
+	MenuBuilder.AddSeparator();
+	
+	const TArray<FName> InitialPluginFilter = GraphObj->GetCurrentPluginFilter();
+
+	TArray<FName> EnabledPluginNames;
+	for (const TSharedRef<IPlugin>& Plugin : IPluginManager::Get().GetEnabledPluginsWithContent())
+	{
+		const FName Name = FName(Plugin.Get().GetName());
+		EnabledPluginNames.AddUnique(Name);
+	}
+
+	TArray<FName> PluginNames = GraphObj->GetEncounteredPluginsAmongNodes();
+	PluginNames.Sort([](const FName& A, const FName& B) { return A.Compare(B) < 0; });
+
+	bool bAddedNonPlugins = false;
+	// First add any "plugins" that are not actual plugins, such as /Game and /Engine.
+	for (const FName& PluginName : PluginNames)
+	{
+		if (EnabledPluginNames.Contains(PluginName))
+		{
+			continue;
+		}
+
+		bAddedNonPlugins = true;
+		
+		PluginFilterAddMenuEntry(MenuBuilder, PluginName);
+	}
+
+	// Separate non-plugins from real plugins.
+	if (bAddedNonPlugins)
+	{
+		MenuBuilder.AddSeparator();
+	}
+
+	// Add all real plugins.
+	for (const FName& PluginName : PluginNames)
+	{
+		if (!EnabledPluginNames.Contains(PluginName))
+		{
+			continue;
+		}
+		
+		PluginFilterAddMenuEntry(MenuBuilder, PluginName);
+	}
+
+	return MenuBuilder.MakeWidget();
+}
+
+FText SReferenceViewer::GetPluginComboButtonText() const
+{
+	const TArray<FName> CurrentPluginFilter = GraphObj->GetCurrentPluginFilter();
+
+	if (CurrentPluginFilter.IsEmpty())
+	{
+		return LOCTEXT("PluginFilterNothingSelected", "None");
+	}
+	else if (CurrentPluginFilter.Num() == 1)
+	{
+		return FText::FromName(CurrentPluginFilter[0]);
+	}
+	else
+	{
+		return LOCTEXT("PluginFilterMultipleSelected", "Multiple");
+	}
+}
+
 void SReferenceViewer::OnShowSoftReferencesChanged()
 {
 	Settings->SetShowSoftReferencesEnabled( !Settings->IsShowSoftReferences() );
diff --git a/Engine/Plugins/Editor/AssetManagerEditor/Source/AssetManagerEditor/Private/ReferenceViewer/SReferenceViewer.h b/Engine/Plugins/Editor/AssetManagerEditor/Source/AssetManagerEditor/Private/ReferenceViewer/SReferenceViewer.h
--- a/Engine/Plugins/Editor/AssetManagerEditor/Source/AssetManagerEditor/Private/ReferenceViewer/SReferenceViewer.h
+++ b/Engine/Plugins/Editor/AssetManagerEditor/Source/AssetManagerEditor/Private/ReferenceViewer/SReferenceViewer.h
@@ -8,6 +8,7 @@
 #include "CollectionManagerTypes.h"
 
 class FUICommandList;
+class SCheckBox;
 class SComboButton;
 class SReferenceViewerFilterBar;
 class SSearchBox;
@@ -128,6 +129,12 @@ private:
 	void HandleCollectionFilterChanged(TSharedPtr<FName> Item, ESelectInfo::Type SelectInfo);
 	FText GetCollectionFilterText() const;
 
+	void OnEnablePluginFilterChanged(ECheckBoxState NewState);
+	ECheckBoxState IsEnablePluginFilterChecked() const;
+	void PluginFilterAddMenuEntry(FMenuBuilder& MenuBuilder, const FName& PluginName);
+	TSharedRef<SWidget> BuildPluginFilterMenu();
+	FText GetPluginComboButtonText() const;
+
 	void OnShowSoftReferencesChanged();
 	bool IsShowSoftReferencesChecked() const;
 	void OnShowHardReferencesChanged();
@@ -232,6 +239,10 @@ private:
 	/** List of collection filter options */
 	TArray<TSharedPtr<FName>> CollectionsComboList;
 
+	/** Combo box for plugin filter options */
+	TSharedPtr<SCheckBox> PluginFilterCheckbox;
+	TSharedPtr<SComboButton> PluginsComboButton;
+
 	/**
 	 * Whether to visually show to the user the option of "Search Depth Limit" or hide it and fix it to a default value:
 	 * - If 0 or negative, it will show to the user the option of "Search Depth Limit".
@@ -246,6 +257,8 @@ private:
 	int32 FixAndHideSearchBreadthLimit;
 	/** Whether to visually show to the user the option of "Collection Filter" */
 	bool bShowCollectionFilter;
+	/** Whether to visually show to the user the option of "Plugin Filter" */
+	bool bShowPluginFilter;
 	/** Whether to visually show to the user the options of "Show Soft/Hard/Management References" */
 	bool bShowShowReferencesOptions;
 	/** Whether to visually show to the user the option of "Show Searchable Names" */
diff --git a/Engine/Plugins/Editor/AssetManagerEditor/Source/AssetManagerEditor/Public/ReferenceViewer/EdGraph_ReferenceViewer.h b/Engine/Plugins/Editor/AssetManagerEditor/Source/AssetManagerEditor/Public/ReferenceViewer/EdGraph_ReferenceViewer.h
--- a/Engine/Plugins/Editor/AssetManagerEditor/Source/AssetManagerEditor/Public/ReferenceViewer/EdGraph_ReferenceViewer.h
+++ b/Engine/Plugins/Editor/AssetManagerEditor/Source/AssetManagerEditor/Public/ReferenceViewer/EdGraph_ReferenceViewer.h
@@ -95,6 +95,10 @@ public:
 	FName GetCurrentCollectionFilter() const;
 	void SetCurrentCollectionFilter(FName NewFilter);
 
+	TArray<FName> GetCurrentPluginFilter() const;
+	void SetCurrentPluginFilter(TArray<FName> NewFilter);
+	TArray<FName> GetEncounteredPluginsAmongNodes() const;
+
 	/* Delegate type to notify when the assets or NodeInfos have changed as opposed to when the filters changed */
 	FSimpleMulticastDelegate& OnAssetsChanged() { return OnAssetsChangedDelegate; }
 
@@ -147,9 +151,16 @@ private:
 
 	/** Returns true if filtering is enabled and we have a valid collection */
 	bool ShouldFilterByCollection() const;
+	
+	/** Returns true if filtering is enabled and we have a valid plugin name filter set */
+	bool ShouldFilterByPlugin() const;
+	
+	void GetUnfilteredGraphPluginNamesRecursive(bool bReferencers, const FAssetIdentifier& InAssetIdentifier, int32 InCurrentDepth, int32 InMaxDepth, const FAssetManagerDependencyQuery& Query, TArray<FAssetIdentifier>& OutAssetIdentifiers);
+	void GetUnfilteredGraphPluginNames(TArray<FAssetIdentifier> RootIdentifiers, TArray<FName>& OutPluginNames);
 
 	void GetSortedLinks(const TArray<FAssetIdentifier>& Identifiers, bool bReferencers, const FAssetManagerDependencyQuery& Query, TMap<FAssetIdentifier, EDependencyPinCategory>& OutLinks) const;
 	bool IsPackageIdentifierPassingFilter(const FAssetIdentifier& InAssetIdentifier) const;
+	bool IsPackageIdentifierPassingPluginFilter(const FAssetIdentifier& InAssetIdentifier) const;
 	bool IsAssetPassingSearchTextFilter(const FAssetIdentifier& InAssetIdentifier) const;
 
 	UEdGraphNode_Reference* FindPath(const FAssetIdentifier& RootId, const FAssetIdentifier& TargetId);
@@ -173,6 +184,12 @@ private:
 	/** Current collection filter. NAME_None for no filter */
 	FName CurrentCollectionFilter;
 
+	/** Current plugin filter. Empty for no filter. */
+	TArray<FName> CurrentPluginFilter;
+
+	/** Plugin names found among unfiltered nodes. Chose among these when filtering for plugins. */
+	TArray<FName> EncounteredPluginsAmongNodes;
+
 	/** A set of the unique class types referenced */
 	TSet<FTopLevelAssetPath> CurrentClasses;
 
diff --git a/Engine/Plugins/Editor/AssetManagerEditor/Source/AssetManagerEditor/Public/ReferenceViewer/ReferenceViewerSettings.h b/Engine/Plugins/Editor/AssetManagerEditor/Source/AssetManagerEditor/Public/ReferenceViewer/ReferenceViewerSettings.h
--- a/Engine/Plugins/Editor/AssetManagerEditor/Source/AssetManagerEditor/Public/ReferenceViewer/ReferenceViewerSettings.h
+++ b/Engine/Plugins/Editor/AssetManagerEditor/Source/AssetManagerEditor/Public/ReferenceViewer/ReferenceViewerSettings.h
@@ -62,6 +62,9 @@ public:
 	bool GetEnableCollectionFilter() const;
 	void SetEnableCollectionFilter(bool bEnabled);
 
+	bool GetEnablePluginFilter() const;
+	void SetEnablePluginFilter(bool bEnabled);
+
 	bool IsShowSoftReferences() const;
 	void SetShowSoftReferencesEnabled(bool bNewEnabled);
 
@@ -137,6 +140,10 @@ private:
 	UPROPERTY(config)
 	bool bEnableCollectionFilter;
 	
+	/* Whether or not to filter for assets by plugins. */
+	UPROPERTY(config)
+	bool bEnablePluginFilter;
+	
 	/* Show/Hide Soft References */
 	UPROPERTY(config)
 	bool bIsShowSoftReferences;
diff --git a/Engine/Source/Runtime/CoreUObject/Public/AssetRegistry/AssetData.h b/Engine/Source/Runtime/CoreUObject/Public/AssetRegistry/AssetData.h
--- a/Engine/Source/Runtime/CoreUObject/Public/AssetRegistry/AssetData.h
+++ b/Engine/Source/Runtime/CoreUObject/Public/AssetRegistry/AssetData.h
@@ -1039,6 +1039,8 @@ struct FReferenceViewerParams
 	int32 FixAndHideSearchBreadthLimit;
 	/** Whether to visually show to the user the option of "Collection Filter" */
 	bool bShowCollectionFilter;
+	/** Whether to visually show to the user the option of "Plugin Filter" */
+	bool bShowPluginFilter;
 	/** Whether to visually show to the user the options of "Show Soft/Hard/Management References" */
 	bool bShowShowReferencesOptions;
 	/** Whether to visually show to the user the option of "Show Searchable Names" */
