Asset Disk Size Tool: Making the Unique, Shared, and External Dependency Size columns non-sortable due to their very high calculation cost and the need to consider all table rows for the sort.

diff --git a/Engine/Plugins/Editor/AssetManagerEditor/Source/AssetManagerEditor/Private/Insights/Table/Widgets/STableTreeView.cpp b/Engine/Plugins/Editor/AssetManagerEditor/Source/AssetManagerEditor/Private/Insights/Table/Widgets/STableTreeView.cpp
--- a/Engine/Plugins/Editor/AssetManagerEditor/Source/AssetManagerEditor/Private/Insights/Table/Widgets/STableTreeView.cpp
+++ b/Engine/Plugins/Editor/AssetManagerEditor/Source/AssetManagerEditor/Private/Insights/Table/Widgets/STableTreeView.cpp
@@ -2857,7 +2857,6 @@ void STableTreeView::ShowColumn(FTableColumn& Column)
 			.VAlignCell(VAlign_Fill)
 			.InitialSortMode(Column.GetInitialSortMode())
 			.SortMode(this, &STableTreeView::GetSortModeForColumn, Column.GetId())
-			.OnSort(this, &STableTreeView::OnSortModeChanged)
 			.FillWidth(Column.GetInitialWidth())
 			//.FixedWidth(Column.IsFixedWidth() ? Column.GetInitialWidth() : TOptional<float>())
 			.HeaderContent()
@@ -2876,6 +2875,11 @@ void STableTreeView::ShowColumn(FTableColumn& Column)
 				TreeViewHeaderRow_GenerateColumnMenu(Column)
 			];
 
+		if (Column.CanBeSorted())
+		{
+			ColumnArgs.OnSort(this, &STableTreeView::OnSortModeChanged);
+		}
+
 		int32 ColumnIndex = 0;
 		const int32 NewColumnPosition = Table->GetColumnPositionIndex(Column.GetId());
 		const int32 NumColumns = TreeViewHeaderRow->GetColumns().Num();
diff --git a/Engine/Plugins/Editor/AssetManagerEditor/Source/AssetManagerEditor/Private/TreeView/AssetTable.cpp b/Engine/Plugins/Editor/AssetManagerEditor/Source/AssetManagerEditor/Private/TreeView/AssetTable.cpp
--- a/Engine/Plugins/Editor/AssetManagerEditor/Source/AssetManagerEditor/Private/TreeView/AssetTable.cpp
+++ b/Engine/Plugins/Editor/AssetManagerEditor/Source/AssetManagerEditor/Private/TreeView/AssetTable.cpp
@@ -953,12 +953,10 @@ void FAssetTable::AddDefaultColumns()
 		TSharedRef<ITableCellValueFormatter> Formatter = MakeShared<FInt64ValueFormatterAsMemory>();
 		Column.SetValueFormatter(Formatter);
 
-		TSharedRef<ITableCellValueSorter> Sorter = MakeShared<FSorterByInt64Value>(ColumnRef);
-		Column.SetValueSorter(Sorter);
-		Column.SetInitialSortMode(EColumnSortMode::Descending);
-
-		//TSharedRef<IFilterValueConverter> Converter = MakeShared<FMemoryFilterValueConverter>();
-		//Column.SetValueConverter(Converter);
+		// Disable sorting on this column for now as it's quite expensive
+		//TSharedRef<ITableCellValueSorter> Sorter = MakeShared<FSorterByInt64Value>(ColumnRef);
+		//Column.SetValueSorter(Sorter);
+		//Column.SetInitialSortMode(EColumnSortMode::Descending);
 
 		Column.SetAggregation(ETableColumnAggregation::None);
 
@@ -1019,12 +1017,10 @@ void FAssetTable::AddDefaultColumns()
 		TSharedRef<ITableCellValueFormatter> Formatter = MakeShared<FInt64ValueFormatterAsMemory>();
 		Column.SetValueFormatter(Formatter);
 
-		TSharedRef<ITableCellValueSorter> Sorter = MakeShared<FSorterByInt64Value>(ColumnRef);
-		Column.SetValueSorter(Sorter);
-		Column.SetInitialSortMode(EColumnSortMode::Descending);
-
-		//TSharedRef<IFilterValueConverter> Converter = MakeShared<FMemoryFilterValueConverter>();
-		//Column.SetValueConverter(Converter);
+		// Disable sorting on this column for now as it's quite expensive
+		//TSharedRef<ITableCellValueSorter> Sorter = MakeShared<FSorterByInt64Value>(ColumnRef);
+		//Column.SetValueSorter(Sorter);
+		//Column.SetInitialSortMode(EColumnSortMode::Descending);
 
 		Column.SetAggregation(ETableColumnAggregation::None);
 
@@ -1085,12 +1081,10 @@ void FAssetTable::AddDefaultColumns()
 		TSharedRef<ITableCellValueFormatter> Formatter = MakeShared<FInt64ValueFormatterAsMemory>();
 		Column.SetValueFormatter(Formatter);
 
-		TSharedRef<ITableCellValueSorter> Sorter = MakeShared<FSorterByInt64Value>(ColumnRef);
-		Column.SetValueSorter(Sorter);
-		Column.SetInitialSortMode(EColumnSortMode::Descending);
-
-		//TSharedRef<IFilterValueConverter> Converter = MakeShared<FMemoryFilterValueConverter>();
-		//Column.SetValueConverter(Converter);
+		// Disable sorting on this column for now as it's quite expensive
+		//TSharedRef<ITableCellValueSorter> Sorter = MakeShared<FSorterByInt64Value>(ColumnRef);
+		//Column.SetValueSorter(Sorter);
+		//Column.SetInitialSortMode(EColumnSortMode::Descending);
 
 		Column.SetAggregation(ETableColumnAggregation::None);
 
