Fix to LOCTEXT conflicting macros. Each ids must be unique if the string do not match.

#rnx

diff --git a/Engine/Plugins/Editor/AssetManagerEditor/Source/AssetManagerEditor/Private/ReferenceViewer/SReferenceViewer.cpp b/Engine/Plugins/Editor/AssetManagerEditor/Source/AssetManagerEditor/Private/ReferenceViewer/SReferenceViewer.cpp
--- a/Engine/Plugins/Editor/AssetManagerEditor/Source/AssetManagerEditor/Private/ReferenceViewer/SReferenceViewer.cpp
+++ b/Engine/Plugins/Editor/AssetManagerEditor/Source/AssetManagerEditor/Private/ReferenceViewer/SReferenceViewer.cpp
@@ -1175,7 +1175,7 @@ TSharedRef<SWidget> SReferenceViewer::BuildPluginFilterMenu()
 
 	MenuBuilder.AddMenuEntry(
 		LOCTEXT("PluginFilterSelectAll", "Select All"),
-		LOCTEXT("PluginFilterSelectAll", "Select all plugins."),
+		LOCTEXT("PluginFilterSelectAllPlugins", "Select all plugins."),
 		FSlateIcon(),
 		FExecuteAction::CreateLambda([this]()
 		{
@@ -1189,7 +1189,7 @@ TSharedRef<SWidget> SReferenceViewer::BuildPluginFilterMenu()
 	
 	MenuBuilder.AddMenuEntry(
 		LOCTEXT("PluginFilterSelectNone", "Select None"),
-		LOCTEXT("PluginFilterSelectNone", "Select no plugins."),
+		LOCTEXT("PluginFilterSelectNoPlugins", "Select no plugins."),
 		FSlateIcon(),
 		FExecuteAction::CreateLambda([this]()
 		{
