Re-instated the asset colors in the Asset Disk Size tool but limited them bewteen Blue and Green ( no brown! )

diff --git a/Engine/Plugins/Editor/AssetManagerEditor/Source/AssetManagerEditor/Private/TreeView/AssetTreeNode.cpp b/Engine/Plugins/Editor/AssetManagerEditor/Source/AssetManagerEditor/Private/TreeView/AssetTreeNode.cpp
--- a/Engine/Plugins/Editor/AssetManagerEditor/Source/AssetManagerEditor/Private/TreeView/AssetTreeNode.cpp
+++ b/Engine/Plugins/Editor/AssetManagerEditor/Source/AssetManagerEditor/Private/TreeView/AssetTreeNode.cpp
@@ -70,8 +70,7 @@ FLinearColor FAssetTreeNode::GetColor(EStyle Style) const
 
 		case EStyle::EAsset:
 		{
-			//return GetAssetChecked().GetColor();
-			return FLinearColor(1.0f, 1.0f, 1.0f, 1.0f);
+			return GetAssetChecked().GetColor();
 			break;
 		}
 
diff --git a/Engine/Plugins/Editor/AssetManagerEditor/Source/AssetManagerEditor/Private/TreeView/SAssetTableTreeView.cpp b/Engine/Plugins/Editor/AssetManagerEditor/Source/AssetManagerEditor/Private/TreeView/SAssetTableTreeView.cpp
--- a/Engine/Plugins/Editor/AssetManagerEditor/Source/AssetManagerEditor/Private/TreeView/SAssetTableTreeView.cpp
+++ b/Engine/Plugins/Editor/AssetManagerEditor/Source/AssetManagerEditor/Private/TreeView/SAssetTableTreeView.cpp
@@ -300,7 +300,7 @@ void SAssetTableTreeView::InitAvailableViewPresets()
 			InOutConfigSet.Add({ FAssetTableColumns::StagedCompressedSizeColumnId,          true, 100.0f });
 			InOutConfigSet.Add({ FAssetTableColumns::PluginInclusiveSizeColumnId,	        true, 100.0f });
 			InOutConfigSet.Add({ FAssetTableColumns::TypeColumnId,                          true, 200.0f });
-			InOutConfigSet.Add({ FAssetTableColumns::NameColumnId,                          true, 200.0f });
+			InOutConfigSet.Add({ FAssetTableColumns::NameColumnId,                         !true, 200.0f });
 			InOutConfigSet.Add({ FAssetTableColumns::PathColumnId,                         !true, 400.0f });
 			InOutConfigSet.Add({ FAssetTableColumns::PrimaryTypeColumnId,                  !true, 200.0f });
 			InOutConfigSet.Add({ FAssetTableColumns::PrimaryNameColumnId,                  !true, 200.0f });
@@ -380,7 +380,8 @@ void SAssetTableTreeView::InitAvailableViewPresets()
 		{
 			InOutConfigSet.Add({ UE::Insights::FTable::GetHierarchyColumnId(),              true, 400.0f });
 			InOutConfigSet.Add({ FAssetTableColumns::CountColumnId,                         true, 100.0f });
-			InOutConfigSet.Add({ FAssetTableColumns::TypeColumnId,                         !true, 200.0f });
+			InOutConfigSet.Add({ FAssetTableColumns::StagedCompressedSizeColumnId,          true, 100.0f });
+			InOutConfigSet.Add({ FAssetTableColumns::TypeColumnId,                          true, 200.0f });
 			InOutConfigSet.Add({ FAssetTableColumns::NameColumnId,                         !true, 200.0f });
 			InOutConfigSet.Add({ FAssetTableColumns::PathColumnId,                         !true, 400.0f });
 			InOutConfigSet.Add({ FAssetTableColumns::PrimaryTypeColumnId,                   true, 200.0f });
@@ -454,7 +455,7 @@ void SAssetTableTreeView::InitAvailableViewPresets()
 			InOutConfigSet.Add({ FAssetTableColumns::CountColumnId,                         true, 100.0f });
 			InOutConfigSet.Add({ FAssetTableColumns::StagedCompressedSizeColumnId,          true, 100.0f });
 			InOutConfigSet.Add({ FAssetTableColumns::TypeColumnId,                          true, 200.0f });
-			InOutConfigSet.Add({ FAssetTableColumns::NameColumnId,                          true, 200.0f });
+			InOutConfigSet.Add({ FAssetTableColumns::NameColumnId,                         !true, 200.0f });
 			InOutConfigSet.Add({ FAssetTableColumns::PathColumnId,                         !true, 400.0f });
 			InOutConfigSet.Add({ FAssetTableColumns::PrimaryTypeColumnId,                   true, 200.0f });
 			InOutConfigSet.Add({ FAssetTableColumns::PrimaryNameColumnId,                   true, 200.0f });
@@ -525,7 +526,7 @@ void SAssetTableTreeView::InitAvailableViewPresets()
 			InOutConfigSet.Add({ FAssetTableColumns::CountColumnId,                         true, 100.0f });
 			InOutConfigSet.Add({ FAssetTableColumns::StagedCompressedSizeColumnId,          true, 100.0f });
 			InOutConfigSet.Add({ FAssetTableColumns::TypeColumnId,                          true, 200.0f });
-			InOutConfigSet.Add({ FAssetTableColumns::NameColumnId,                          true, 200.0f });
+			InOutConfigSet.Add({ FAssetTableColumns::NameColumnId,                         !true, 200.0f });
 			InOutConfigSet.Add({ FAssetTableColumns::PathColumnId,                         !true, 400.0f });
 			InOutConfigSet.Add({ FAssetTableColumns::PrimaryTypeColumnId,                   true, 200.0f });
 			InOutConfigSet.Add({ FAssetTableColumns::PrimaryNameColumnId,                   true, 200.0f });
@@ -597,7 +598,7 @@ void SAssetTableTreeView::InitAvailableViewPresets()
 			InOutConfigSet.Add({ FAssetTableColumns::CountColumnId,                         true, 100.0f });
 			InOutConfigSet.Add({ FAssetTableColumns::StagedCompressedSizeColumnId,          true, 100.0f });
 			InOutConfigSet.Add({ FAssetTableColumns::TypeColumnId,                          true, 200.0f });
-			InOutConfigSet.Add({ FAssetTableColumns::NameColumnId,                          true, 200.0f });
+			InOutConfigSet.Add({ FAssetTableColumns::NameColumnId,                         !true, 200.0f });
 			InOutConfigSet.Add({ FAssetTableColumns::PathColumnId,                         !true, 400.0f });
 			InOutConfigSet.Add({ FAssetTableColumns::PrimaryTypeColumnId,                   true, 200.0f });
 			InOutConfigSet.Add({ FAssetTableColumns::PrimaryNameColumnId,                   true, 200.0f });
@@ -1384,7 +1385,7 @@ void SAssetTableTreeView::PopulateAssetTableRow(FAssetTableRow& OutRow, const FA
 
 	// Sets the color based on asset type.
 	const uint32 AssetTypeHash = GetTypeHash(FStringView(OutRow.Type));
-	OutRow.Color = USlateThemeManager::Get().GetColor((EStyleColor)((uint32)EStyleColor::AccentBlue + AssetTypeHash % 8));
+	OutRow.Color = USlateThemeManager::Get().GetColor((EStyleColor)((uint32)EStyleColor::AccentBlue + AssetTypeHash % ((uint32)EStyleColor::AccentGreen - (uint32)EStyleColor::AccentBlue)));
 }
 
 ////////////////////////////////////////////////////////////////////////////////////////////////////
