Changed hard coded UI colors for Asset Disk Size tool over to use standard SlateTheme Acent Colors
Split the Icon Color and Text Color from the BaseTreeNode to give more control over the look.
Text is now always white and only the Icon is colored by type

diff --git a/Engine/Plugins/Editor/AssetManagerEditor/Source/AssetManagerEditor/Private/Insights/Table/ViewModels/BaseTreeNode.cpp b/Engine/Plugins/Editor/AssetManagerEditor/Source/AssetManagerEditor/Private/Insights/Table/ViewModels/BaseTreeNode.cpp
--- a/Engine/Plugins/Editor/AssetManagerEditor/Source/AssetManagerEditor/Private/Insights/Table/ViewModels/BaseTreeNode.cpp
+++ b/Engine/Plugins/Editor/AssetManagerEditor/Source/AssetManagerEditor/Private/Insights/Table/ViewModels/BaseTreeNode.cpp
@@ -76,11 +76,11 @@ FLinearColor FBaseTreeNode::GetDefaultColor(bool bIsGroupNode)
 {
 	if (bIsGroupNode)
 	{
-		return FLinearColor(1.0f, 0.9f, 0.6f, 1.0f);
+		return USlateThemeManager::Get().GetColor(EStyleColor::AccentYellow);
 	}
 	else
 	{
-		return FLinearColor(1.0f, 1.0f, 1.0f, 1.0f);
+		return USlateThemeManager::Get().GetColor(EStyleColor::AccentWhite);
 	}
 }
 
diff --git a/Engine/Plugins/Editor/AssetManagerEditor/Source/AssetManagerEditor/Private/Insights/Table/ViewModels/BaseTreeNode.h b/Engine/Plugins/Editor/AssetManagerEditor/Source/AssetManagerEditor/Private/Insights/Table/ViewModels/BaseTreeNode.h
--- a/Engine/Plugins/Editor/AssetManagerEditor/Source/AssetManagerEditor/Private/Insights/Table/ViewModels/BaseTreeNode.h
+++ b/Engine/Plugins/Editor/AssetManagerEditor/Source/AssetManagerEditor/Private/Insights/Table/ViewModels/BaseTreeNode.h
@@ -121,17 +121,25 @@ public:
 		return GetDefaultIcon(IsGroup());
 	}
 
+	/**
+	 * @returns the color tint to be used for the icon used for this node.
+	 */
+	virtual FLinearColor GetIconColor() const
+	{
+		return GetDefaultColor(IsGroup());
+	}
+
 	/**
 	 * @returns the default color tint for a group/leaf node.
 	 */
 	static FLinearColor GetDefaultColor(bool bIsGroupNode);
 
 	/**
-	 * @returns the color tint to be used for the icon and the name text of this node.
+	 * @returns the color tint to be used for the name text of this node.
 	 */
 	virtual FLinearColor GetColor() const
 	{
-		return GetDefaultColor(IsGroup());
+		return FLinearColor(1.0f, 1.0f, 1.0f, 1.0f);
 	}
 
 	//////////////////////////////////////////////////
diff --git a/Engine/Plugins/Editor/AssetManagerEditor/Source/AssetManagerEditor/Private/Insights/Table/Widgets/STableTreeViewCell.cpp b/Engine/Plugins/Editor/AssetManagerEditor/Source/AssetManagerEditor/Private/Insights/Table/Widgets/STableTreeViewCell.cpp
--- a/Engine/Plugins/Editor/AssetManagerEditor/Source/AssetManagerEditor/Private/Insights/Table/Widgets/STableTreeViewCell.cpp
+++ b/Engine/Plugins/Editor/AssetManagerEditor/Source/AssetManagerEditor/Private/Insights/Table/Widgets/STableTreeViewCell.cpp
@@ -183,8 +183,8 @@ FSlateColor STableTreeViewCell::GetIconColorAndOpacity() const
 	bool bIsHoveredOrSelected = IsHovered() || IsSelected();
 
 	return bIsHoveredOrSelected ?
-		TableTreeNodePtr->GetColor() :
-		TableTreeNodePtr->GetColor().CopyWithNewOpacity(0.5f);
+		TableTreeNodePtr->GetIconColor() :
+		TableTreeNodePtr->GetIconColor().CopyWithNewOpacity(0.5f);
 }
 
 ////////////////////////////////////////////////////////////////////////////////////////////////////
diff --git a/Engine/Plugins/Editor/AssetManagerEditor/Source/AssetManagerEditor/Private/TreeView/AssetTable.cpp b/Engine/Plugins/Editor/AssetManagerEditor/Source/AssetManagerEditor/Private/TreeView/AssetTable.cpp
--- a/Engine/Plugins/Editor/AssetManagerEditor/Source/AssetManagerEditor/Private/TreeView/AssetTable.cpp
+++ b/Engine/Plugins/Editor/AssetManagerEditor/Source/AssetManagerEditor/Private/TreeView/AssetTable.cpp
@@ -686,6 +686,57 @@
 
 		AddColumn(ColumnRef);
 	}
+
 	//////////////////////////////////////////////////
+	// PluginInclusiveSize Column
+	{
+		TSharedRef<FTableColumn> ColumnRef = MakeShared<FTableColumn>(FAssetTableColumns::PluginInclusiveSizeColumnId);
+		FTableColumn& Column = *ColumnRef;
+
+		Column.SetIndex(ColumnIndex++);
+
+		Column.SetShortName(LOCTEXT("PluginInclusiveSizeColumnName", "Incl. Size Plugin"));
+		Column.SetTitleName(LOCTEXT("PluginInclusiveSizeColumnTitle", "Plugin Inclusive Size"));
+		Column.SetDescription(LOCTEXT("PluginInclusiveColumnDesc", "Inclusive size of this plugin and its dependencies"));
+
+		Column.SetFlags(ETableColumnFlags::ShouldBeVisible | ETableColumnFlags::CanBeHidden | ETableColumnFlags::CanBeFiltered);
+
+		Column.SetHorizontalAlignment(HAlign_Left);
+		Column.SetInitialWidth(50.0f);
+
+		Column.SetDataType(ETableCellDataType::Int64);
+
+		class FPluginInclusiveSizeValueGetter : public FTableCellValueGetter
+		{
+		public:
+			virtual const TOptional<FTableCellValue> GetValue(const FTableColumn& Column, const FBaseTreeNode& Node) const override
+			{
+				if (Node.Is<FPluginSimpleGroupNode>())
+				{
+					TSharedPtr<FTable> TablePtr = static_cast<const FPluginSimpleGroupNode&>(Node).GetParentTable().Pin();
+					const FAssetTable& AssetTable = static_cast<const FAssetTable&>(*TablePtr);
+					int32 PluginIndex = static_cast<const FPluginSimpleGroupNode&>(Node).GetPluginIndex();
+					return FTableCellValue(AssetTable.GetPluginInfoByIndex(PluginIndex).GetOrComputeTotalSizeInclusiveOfDependencies(AssetTable));
+				}
+
+				return TOptional<FTableCellValue>();
+			}
+		};
+		TSharedRef<ITableCellValueGetter> Getter = MakeShared<FPluginInclusiveSizeValueGetter>();
+		Column.SetValueGetter(Getter);
+
+		TSharedRef<ITableCellValueFormatter> Formatter = MakeShared<FInt64ValueFormatterAsMemory>();
+		Column.SetValueFormatter(Formatter);
+
+		TSharedRef<ITableCellValueSorter> Sorter = MakeShared<FSorterByInt64Value>(ColumnRef);
+		Column.SetValueSorter(Sorter);
+		Column.SetInitialSortMode(EColumnSortMode::Descending);
+
+		Column.SetAggregation(ETableColumnAggregation::None);
+
+		AddColumn(ColumnRef);
+	}
+
+	//////////////////////////////////////////////////
 	// Type Column
 	{
@@ -1272,55 +1323,7 @@
 
 		AddColumn(ColumnRef);
 	}
-	//////////////////////////////////////////////////
-	// PluginInclusiveSize Column
-	{
-		TSharedRef<FTableColumn> ColumnRef = MakeShared<FTableColumn>(FAssetTableColumns::PluginInclusiveSizeColumnId);
-		FTableColumn& Column = *ColumnRef;
-
-		Column.SetIndex(ColumnIndex++);
-
-		Column.SetShortName(LOCTEXT("PluginInclusiveSizeColumnName", "Incl. Size Plugin"));
-		Column.SetTitleName(LOCTEXT("PluginInclusiveSizeColumnTitle", "Plugin Inclusive Size"));
-		Column.SetDescription(LOCTEXT("PluginInclusiveColumnDesc", "Inclusive size of this plugin and its dependencies"));
-
-		Column.SetFlags(ETableColumnFlags::ShouldBeVisible | ETableColumnFlags::CanBeHidden | ETableColumnFlags::CanBeFiltered);
-
-		Column.SetHorizontalAlignment(HAlign_Left);
-		Column.SetInitialWidth(50.0f);
-
-		Column.SetDataType(ETableCellDataType::Int64);
-
-		class FPluginInclusiveSizeValueGetter : public FTableCellValueGetter
-		{
-		public:
-			virtual const TOptional<FTableCellValue> GetValue(const FTableColumn& Column, const FBaseTreeNode& Node) const override
-			{
-				if (Node.Is<FPluginSimpleGroupNode>())
-				{
-					TSharedPtr<FTable> TablePtr = static_cast<const FPluginSimpleGroupNode&>(Node).GetParentTable().Pin();
-					const FAssetTable& AssetTable = static_cast<const FAssetTable&>(*TablePtr);
-					int32 PluginIndex = static_cast<const FPluginSimpleGroupNode&>(Node).GetPluginIndex();
-					return FTableCellValue(AssetTable.GetPluginInfoByIndex(PluginIndex).GetOrComputeTotalSizeInclusiveOfDependencies(AssetTable));
-				}
-
-				return TOptional<FTableCellValue>();
-			}
-		};
-		TSharedRef<ITableCellValueGetter> Getter = MakeShared<FPluginInclusiveSizeValueGetter>();
-		Column.SetValueGetter(Getter);
-
-		TSharedRef<ITableCellValueFormatter> Formatter = MakeShared<FInt64ValueFormatterAsMemory>();
-		Column.SetValueFormatter(Formatter);
-
-		TSharedRef<ITableCellValueSorter> Sorter = MakeShared<FSorterByInt64Value>(ColumnRef);
-		Column.SetValueSorter(Sorter);
-		Column.SetInitialSortMode(EColumnSortMode::Descending);
-
-		Column.SetAggregation(ETableColumnAggregation::None);
-
-		AddColumn(ColumnRef);
-	}
+	
 }
 
 /*static*/TSet<int32> FAssetTableRow::GatherAllReachableNodes(const TArray<int32>& StartingNodes, const FAssetTable& OwningTable, const TSet<int32>& AdditionalNodesToStopAt, const TSet<const TCHAR*, TStringPointerSetKeyFuncs_DEPRECATED<const TCHAR*>>& RestrictToPlugins, TMap<int32, TArray<int32>>* OutRouteMap /*= nullptr*/)
diff --git a/Engine/Plugins/Editor/AssetManagerEditor/Source/AssetManagerEditor/Private/TreeView/AssetTreeNode.cpp b/Engine/Plugins/Editor/AssetManagerEditor/Source/AssetManagerEditor/Private/TreeView/AssetTreeNode.cpp
--- a/Engine/Plugins/Editor/AssetManagerEditor/Source/AssetManagerEditor/Private/TreeView/AssetTreeNode.cpp
+++ b/Engine/Plugins/Editor/AssetManagerEditor/Source/AssetManagerEditor/Private/TreeView/AssetTreeNode.cpp
@@ -28,70 +28,65 @@ const FSlateBrush* FAssetTreeNode::GetIcon(EStyle Style) const
 		case EStyle::EDefault:
 		{
 			return UE::Insights::FInsightsStyle::GetBrush("Icons.Leaf.TreeItem");
-			break;
 		}
 
 		case EStyle::EAsset:
 		{
 			return UE::Insights::FInsightsStyle::GetBrush("Icons.Asset.TreeItem");
-			break;
 		}
 				
 		case EStyle::EGroup:
 		{
 			return UE::Insights::FInsightsStyle::GetBrush("Icons.Group.TreeItem");
-			break;
 		}	
 
 		case EStyle::EDependencies:
 		{
 			return UE::Insights::FInsightsStyle::GetBrush("Icons.Dependencies.TreeItem");
-			break;
 		}
 
 		case EStyle::EPlugin:
 		{
 			return UE::Insights::FInsightsStyle::GetBrush("Icons.Plugin.TreeItem");
-			break;
 		}
 	}
 }
 
-FLinearColor FAssetTreeNode::GetColor(EStyle Style) const 
+FLinearColor FAssetTreeNode::GetIconColor(EStyle Style) const
 {
 	switch (Style)
 	{
 		default:
 		case EStyle::EDefault:
 		{
-			return FLinearColor(1.0f, 1.0f, 0.5f, 1.0f);
-			break;
+			return USlateThemeManager::Get().GetColor(EStyleColor::AccentWhite);
 		}
 
 		case EStyle::EAsset:
 		{
 			return GetAssetChecked().GetColor();
-			break;
 		}
 
 		case EStyle::EGroup:
 		{
-			return FLinearColor(1.0f, 1.0f, 0.5f, 1.0f);
-			break;
+			return USlateThemeManager::Get().GetColor(EStyleColor::AccentYellow);
 		}
 
 		case EStyle::EDependencies:
 		{
-			return FLinearColor(1.0f, 1.0f, 1.0f, 1.0f);
-			break;
+			return USlateThemeManager::Get().GetColor(EStyleColor::AccentWhite);
 		}
 
 		case EStyle::EPlugin:
 		{
-			return FLinearColor(1.0f, 1.0f, 1.0f, 1.0f);
-			break;
-		}	
-	}	
+			return USlateThemeManager::Get().GetColor(EStyleColor::AccentGreen);
+		}
+	}
+}
+
+FLinearColor FAssetTreeNode::GetColor(EStyle Style) const 
+{
+	return USlateThemeManager::Get().GetColor(EStyleColor::AccentWhite);
 }
 
 FAssetTreeNode::EStyle FAssetTreeNode::GetStyle() const
@@ -104,6 +99,11 @@ const FSlateBrush* FAssetTreeNode::GetIcon() const
 	return GetIcon(GetStyle());
 }
 
+FLinearColor FAssetTreeNode::GetIconColor() const
+{
+	return GetIconColor(GetStyle());
+}
+
 ////////////////////////////////////////////////////////////////////////////////////////////////////
 
 FLinearColor FAssetTreeNode::GetColor() const
diff --git a/Engine/Plugins/Editor/AssetManagerEditor/Source/AssetManagerEditor/Private/TreeView/AssetTreeNode.h b/Engine/Plugins/Editor/AssetManagerEditor/Source/AssetManagerEditor/Private/TreeView/AssetTreeNode.h
--- a/Engine/Plugins/Editor/AssetManagerEditor/Source/AssetManagerEditor/Private/TreeView/AssetTreeNode.h
+++ b/Engine/Plugins/Editor/AssetManagerEditor/Source/AssetManagerEditor/Private/TreeView/AssetTreeNode.h
@@ -44,6 +44,7 @@ public:
 	const FAssetTableRow& GetAssetChecked() const { return AssetTablePtr->GetAssetChecked(RowId.RowIndex); }
 
 	virtual const FSlateBrush* GetIcon() const final;
+	virtual FLinearColor GetIconColor() const final;
 	virtual FLinearColor GetColor() const final;
 	
 protected:
@@ -60,6 +61,7 @@ protected:
 
 	virtual EStyle GetStyle() const;
 	const FSlateBrush* GetIcon(EStyle Style) const;
+	FLinearColor GetIconColor(EStyle Style) const;
 	FLinearColor GetColor(EStyle Style) const;
 
 private:
