Non-Unity fix for AssetTreeNode.cpp

diff --git a/Engine/Plugins/Editor/AssetManagerEditor/Source/AssetManagerEditor/Private/TreeView/AssetTreeNode.cpp b/Engine/Plugins/Editor/AssetManagerEditor/Source/AssetManagerEditor/Private/TreeView/AssetTreeNode.cpp
--- a/Engine/Plugins/Editor/AssetManagerEditor/Source/AssetManagerEditor/Private/TreeView/AssetTreeNode.cpp
+++ b/Engine/Plugins/Editor/AssetManagerEditor/Source/AssetManagerEditor/Private/TreeView/AssetTreeNode.cpp
@@ -7,6 +7,7 @@
 #include "Insights/Common/InsightsStyle.h"
 #include "Insights/Table/Widgets/STableTreeView.h"
 #include "Internationalization/Internationalization.h"
+#include "Styling/StyleColors.h"
 
 #define LOCTEXT_NAMESPACE "FAssetTreeNode"
 
