As agreed, renaming "BaseGame" for the logical plugin which maps to /Game/ to "Game" in Asset Disk Size tool so it is consistent with the convention in the build and telemetry

diff --git a/Engine/Plugins/Editor/AssetManagerEditor/Source/AssetManagerEditor/Private/AssetManagerEditorModule.cpp b/Engine/Plugins/Editor/AssetManagerEditor/Source/AssetManagerEditor/Private/AssetManagerEditorModule.cpp
--- a/Engine/Plugins/Editor/AssetManagerEditor/Source/AssetManagerEditor/Private/AssetManagerEditorModule.cpp
+++ b/Engine/Plugins/Editor/AssetManagerEditor/Source/AssetManagerEditor/Private/AssetManagerEditorModule.cpp
@@ -1352,7 +1352,7 @@ bool FAssetManagerEditorModule::GetStringValueForCustomColumn(const FAssetData&
 		const FString& PackageNameString = AssetData.PackageName.ToString();
 		if (PackageNameString.StartsWith(TEXT("/Game/")))
 		{
-			OutValue = TEXT("BaseGame");
+			OutValue = TEXT("Game");
 			return true;
 		}
 		else if (PackageNameString.ParseIntoArray(ParsedPath, TEXT("/")) >= 1)
