Further UX improvements to Highlightin of Icons and Text

diff --git a/Engine/Plugins/Editor/AssetManagerEditor/Source/AssetManagerEditor/Private/Insights/Table/Widgets/STableTreeViewCell.cpp b/Engine/Plugins/Editor/AssetManagerEditor/Source/AssetManagerEditor/Private/Insights/Table/Widgets/STableTreeViewCell.cpp
--- a/Engine/Plugins/Editor/AssetManagerEditor/Source/AssetManagerEditor/Private/Insights/Table/Widgets/STableTreeViewCell.cpp
+++ b/Engine/Plugins/Editor/AssetManagerEditor/Source/AssetManagerEditor/Private/Insights/Table/Widgets/STableTreeViewCell.cpp
@@ -184,7 +184,7 @@ FSlateColor STableTreeViewCell::GetIconColorAndOpacity() const
 
 	return bIsHoveredOrSelected ?
 		TableTreeNodePtr->GetIconColor() :
-		TableTreeNodePtr->GetIconColor().CopyWithNewOpacity(0.5f);
+		TableTreeNodePtr->GetIconColor().CopyWithNewOpacity(0.8f);
 }
 
 ////////////////////////////////////////////////////////////////////////////////////////////////////
