STableTreeView -- ensure that CreateGroups and Aggregate respect cancelling in progress operations and add a lock to allow us to wait for async tasks to complete before destruction of the STableTreeView itself.

diff --git a/Engine/Plugins/Editor/AssetManagerEditor/Source/AssetManagerEditor/Private/Insights/Table/Widgets/STableTreeView.cpp b/Engine/Plugins/Editor/AssetManagerEditor/Source/AssetManagerEditor/Private/Insights/Table/Widgets/STableTreeView.cpp
--- a/Engine/Plugins/Editor/AssetManagerEditor/Source/AssetManagerEditor/Private/Insights/Table/Widgets/STableTreeView.cpp
+++ b/Engine/Plugins/Editor/AssetManagerEditor/Source/AssetManagerEditor/Private/Insights/Table/Widgets/STableTreeView.cpp
@@ -142,13 +142,16 @@ STableTreeView::STableTreeView()
 	, ColumnBeingSorted(GetDefaultColumnBeingSorted())
 	, ColumnSortMode(GetDefaultColumnSortMode())
 {
+	AsyncOperationMutex = MakeShared<FRWLock>();
 }
 
 ////////////////////////////////////////////////////////////////////////////////////////////////////
 
 STableTreeView::~STableTreeView()
 {
-	ensureMsgf(bIsCloseScheduled, TEXT("TableTreeView running in async mode was closed but OnClose() was not called. This can lead to a crash. Call OnClose() from the owner tab/window."));
+	CancelCurrentAsyncOp();
+
+	FScopedExclusiveRWLock AsyncOpLock(AsyncOperationMutex);
 
 	if (CurrentAsyncOpFilterConfigurator)
 	{
@@ -998,7 +1001,7 @@ FGraphEventRef STableTreeView::StartNodeFilteringTask(FGraphEventRef Prerequisit
 		Prerequisites.Add(DispatchEvent);
 	}
 
-	return TGraphTask<FTableTreeViewNodeFilteringAsyncTask>::CreateTask(&Prerequisites).ConstructAndDispatchWhenReady(this);
+	return TGraphTask<FTableTreeViewNodeFilteringAsyncTask>::CreateTask(&Prerequisites).ConstructAndDispatchWhenReady(this, AsyncOperationMutex);
 }
 
 ////////////////////////////////////////////////////////////////////////////////////////////////////
@@ -1263,7 +1266,7 @@ FGraphEventRef STableTreeView::StartHierarchyFilteringTask(FGraphEventRef Prereq
 		Prerequisites.Add(DispatchEvent);
 	}
 
-	return TGraphTask<FTableTreeViewHierarchyFilteringAsyncTask>::CreateTask(&Prerequisites).ConstructAndDispatchWhenReady(this);
+	return TGraphTask<FTableTreeViewHierarchyFilteringAsyncTask>::CreateTask(&Prerequisites).ConstructAndDispatchWhenReady(this, AsyncOperationMutex);
 }
 
 ////////////////////////////////////////////////////////////////////////////////////////////////////
@@ -1661,7 +1664,7 @@ FGraphEventRef STableTreeView::StartGroupingTask(FGraphEventRef Prerequisite)
 		Prerequisites.Add(DispatchEvent);
 	}
 
-	return TGraphTask<FTableTreeViewGroupingAsyncTask>::CreateTask(&Prerequisites).ConstructAndDispatchWhenReady(this, &CurrentAsyncOpGroupings);
+	return TGraphTask<FTableTreeViewGroupingAsyncTask>::CreateTask(&Prerequisites).ConstructAndDispatchWhenReady(this, &CurrentAsyncOpGroupings, AsyncOperationMutex);
 }
 
 ////////////////////////////////////////////////////////////////////////////////////////////////////
@@ -2536,6 +2539,11 @@ void STableTreeView::UpdateAggregatedValuesSingleNode(FTableTreeNode& InOutGroup
 
 void STableTreeView::UpdateAggregatedValuesRec(FTableTreeNode& InOutGroupNode)
 {
+	if (AsyncOperationProgress.ShouldCancelAsyncOp())
+	{
+		return;
+	}
+
 	STableTreeView::UpdateAggregatedValues<true>(Table, InOutGroupNode);
 }
 
@@ -2662,7 +2670,7 @@ FGraphEventRef STableTreeView::StartSortingTask(FGraphEventRef Prerequisite)
 		Prerequisites.Add(DispatchEvent);
 	}
 
-	return TGraphTask<FTableTreeViewSortingAsyncTask>::CreateTask(&Prerequisites).ConstructAndDispatchWhenReady(this, CurrentAsyncOpSorter, CurrentAsyncOpColumnSortMode);
+	return TGraphTask<FTableTreeViewSortingAsyncTask>::CreateTask(&Prerequisites).ConstructAndDispatchWhenReady(this, CurrentAsyncOpSorter, CurrentAsyncOpColumnSortMode, AsyncOperationMutex);
 }
 
 ////////////////////////////////////////////////////////////////////////////////////////////////////
diff --git a/Engine/Plugins/Editor/AssetManagerEditor/Source/AssetManagerEditor/Private/Insights/Table/Widgets/STableTreeView.h b/Engine/Plugins/Editor/AssetManagerEditor/Source/AssetManagerEditor/Private/Insights/Table/Widgets/STableTreeView.h
--- a/Engine/Plugins/Editor/AssetManagerEditor/Source/AssetManagerEditor/Private/Insights/Table/Widgets/STableTreeView.h
+++ b/Engine/Plugins/Editor/AssetManagerEditor/Source/AssetManagerEditor/Private/Insights/Table/Widgets/STableTreeView.h
@@ -603,6 +603,13 @@ protected:
 	bool bIsUpdateRunning = false;
 	bool bIsCloseScheduled = false;
 
+	// The AsyncOperationLock is used a little differently than typic for a RW lock. 
+	// Instead of thinking of it as allowing multiple readers or one writer, it can be thought of
+	// as multiple shared lockers or one exclusive locker. The async operations all take read/shared locks
+	// and the destructor of the tree view takes an exclusive lock. This ensure that all async ops finish before
+	// the destructor runs and cleans up the data they are using
+	TSharedPtr<FRWLock> AsyncOperationMutex;
+
 	TArray<FTableTreeNodePtr> DummyGroupNodes;
 	FGraphEventRef InProgressAsyncOperationEvent;
 	FGraphEventRef AsyncCompleteTaskEvent;
@@ -629,13 +636,44 @@ protected:
 	FName LogListingName;
 };
 
+////////////////////////////////////////////////////////////////////////////////////////////////////
+ 
+class FScopedExclusiveRWLock
+{
+public:
+	FScopedExclusiveRWLock(TSharedPtr<FRWLock>& InLock)
+		: Lock(InLock) {
+		Lock->WriteLock();
+	}
+
+	~FScopedExclusiveRWLock() { Lock->WriteUnlock(); }
+
+private:
+	TSharedPtr<FRWLock> Lock;
+	};
+
+class FScopedSharedRWLock
+{
+public:
+	FScopedSharedRWLock(TSharedPtr<FRWLock>& InLock)
+		: Lock(InLock) {
+		Lock->ReadLock();
+	}
+
+	~FScopedSharedRWLock() { Lock->ReadUnlock(); }
+
+private:
+	TSharedPtr<FRWLock> Lock;
+};
+
 ////////////////////////////////////////////////////////////////////////////////////////////////////
 
 class FTableTreeViewNodeFilteringAsyncTask
 {
 public:
-	FTableTreeViewNodeFilteringAsyncTask(STableTreeView* InPtr)
+	FTableTreeViewNodeFilteringAsyncTask(STableTreeView* InPtr, TSharedPtr<FRWLock>& InAsyncOpLock)
 	{
+		AsyncOpLock = InAsyncOpLock;
 		TableTreeViewPtr = InPtr;
 	}
 
@@ -645,14 +683,17 @@ public:
 
 	void DoTask(ENamedThreads::Type CurrentThread, const FGraphEventRef& MyCompletionGraphEvent)
 	{
+		FScopedSharedRWLock SharedLock(AsyncOpLock);
 		if (TableTreeViewPtr)
 		{
 			TableTreeViewPtr->ApplyNodeFiltering();
 		}
+		AsyncOpLock.Reset();
 	}
 
 private:
 	STableTreeView* TableTreeViewPtr = nullptr;
+	TSharedPtr<FRWLock> AsyncOpLock;
 };
 
 ////////////////////////////////////////////////////////////////////////////////////////////////////
@@ -660,8 +701,9 @@ private:
 class FTableTreeViewHierarchyFilteringAsyncTask
 {
 public:
-	FTableTreeViewHierarchyFilteringAsyncTask(STableTreeView* InPtr)
+	FTableTreeViewHierarchyFilteringAsyncTask(STableTreeView* InPtr, TSharedPtr<FRWLock> InAsyncOpLock)
 	{
+		AsyncOpLock = InAsyncOpLock;
 		TableTreeViewPtr = InPtr;
 	}
 
@@ -671,14 +713,17 @@ public:
 
 	void DoTask(ENamedThreads::Type CurrentThread, const FGraphEventRef& MyCompletionGraphEvent)
 	{
+		FScopedSharedRWLock ScopedLock(AsyncOpLock);
 		if (TableTreeViewPtr)
 		{
 			TableTreeViewPtr->ApplyHierarchyFiltering();
 		}
+		AsyncOpLock.Reset();
 	}
 
 private:
 	STableTreeView* TableTreeViewPtr = nullptr;
+	TSharedPtr<FRWLock> AsyncOpLock;
 };
 
 ////////////////////////////////////////////////////////////////////////////////////////////////////
@@ -686,8 +731,9 @@ private:
 class FTableTreeViewSortingAsyncTask
 {
 public:
-	FTableTreeViewSortingAsyncTask(STableTreeView* InPtr, ITableCellValueSorter* InSorter, EColumnSortMode::Type InColumnSortMode)
+	FTableTreeViewSortingAsyncTask(STableTreeView* InPtr, ITableCellValueSorter* InSorter, EColumnSortMode::Type InColumnSortMode, TSharedPtr<FRWLock>& InAsyncOpLock)
 	{
+		AsyncOpLock = InAsyncOpLock;
 		TableTreeViewPtr = InPtr;
 		Sorter = InSorter;
 		ColumnSortMode = InColumnSortMode;
@@ -699,16 +745,19 @@ public:
 
 	void DoTask(ENamedThreads::Type CurrentThread, const FGraphEventRef& MyCompletionGraphEvent)
 	{
+		FScopedSharedRWLock ScopedLock(AsyncOpLock);
 		if (TableTreeViewPtr)
 		{
 			TableTreeViewPtr->SortTreeNodes(Sorter, ColumnSortMode);
 		}
+		AsyncOpLock.Reset();
 	}
 
 private:
 	STableTreeView* TableTreeViewPtr;
 	ITableCellValueSorter* Sorter;
 	EColumnSortMode::Type ColumnSortMode;
+	TSharedPtr<FRWLock> AsyncOpLock;
 };
 
 ////////////////////////////////////////////////////////////////////////////////////////////////////
@@ -716,8 +765,9 @@ private:
 class FTableTreeViewGroupingAsyncTask
 {
 public:
-	FTableTreeViewGroupingAsyncTask(STableTreeView* InPtr, TArray<TSharedPtr<FTreeNodeGrouping>>* InGroupings)
+	FTableTreeViewGroupingAsyncTask(STableTreeView* InPtr, TArray<TSharedPtr<FTreeNodeGrouping>>* InGroupings, TSharedPtr<FRWLock>& InAsyncOpLock)
 	{
+		AsyncOpLock = InAsyncOpLock;
 		TableTreeViewPtr = InPtr;
 		Groupings = InGroupings;
 	}
@@ -728,15 +778,18 @@ public:
 
 	void DoTask(ENamedThreads::Type CurrentThread, const FGraphEventRef& MyCompletionGraphEvent)
 	{
+		FScopedSharedRWLock ScopedLock(AsyncOpLock);
 		if (TableTreeViewPtr)
 		{
 			TableTreeViewPtr->CreateGroups(*Groupings);
 		}
+		AsyncOpLock.Reset();
 	}
 
 private:
 	STableTreeView* TableTreeViewPtr = nullptr;
 	TArray<TSharedPtr<FTreeNodeGrouping>>* Groupings;
+	TSharedPtr<FRWLock> AsyncOpLock;
 };
 
 ////////////////////////////////////////////////////////////////////////////////////////////////////
@@ -773,9 +826,10 @@ private:
 class FSearchForItemToSelectTask
 {
 public:
-	FSearchForItemToSelectTask(TSharedPtr<FTableTaskCancellationToken> InToken, TSharedPtr<STableTreeView> InPtr)
+	FSearchForItemToSelectTask(TSharedPtr<FTableTaskCancellationToken> InToken, TSharedPtr<STableTreeView> InPtr, TSharedPtr<FRWLock> InAsyncOpLock)
 		: CancellationToken(InToken)
 		, TableTreeViewPtr(InPtr)
+		, AsyncOpLock(InAsyncOpLock)
 	{}
 
 	FORCEINLINE TStatId GetStatId() const { RETURN_QUICK_DECLARE_CYCLE_STAT(FSearchForItemToSelectTask, STATGROUP_TaskGraphTasks); }
@@ -784,12 +838,18 @@ public:
 
 	void DoTask(ENamedThreads::Type CurrentThread, const FGraphEventRef& MyCompletionGraphEvent)
 	{
-		TableTreeViewPtr->SearchForItem(CancellationToken);
+		FScopedSharedRWLock ScopedLock(AsyncOpLock);
+		if (TableTreeViewPtr)
+		{
+			TableTreeViewPtr->SearchForItem(CancellationToken);
+		}
+		AsyncOpLock.Reset();
 	}
 
 private:
 	TSharedPtr<FTableTaskCancellationToken> CancellationToken;
 	TSharedPtr<STableTreeView> TableTreeViewPtr;
+	TSharedPtr<FRWLock> AsyncOpLock;
 };
 
 ////////////////////////////////////////////////////////////////////////////////////////////////////
@@ -797,9 +857,10 @@ private:
 class FSelectNodeByTableRowIndexTask
 {
 public:
-	FSelectNodeByTableRowIndexTask(TSharedPtr<FTableTaskCancellationToken> InToken, TSharedPtr<STableTreeView> InPtr, uint32 InRowIndex)
+	FSelectNodeByTableRowIndexTask(TSharedPtr<FTableTaskCancellationToken> InToken, TSharedPtr<STableTreeView> InPtr, uint32 InRowIndex, TSharedPtr<FRWLock> InAsyncOpLock)
 		: CancellationToken(InToken)
 		, TableTreeViewPtr(InPtr)
+		, AsyncOpLock(InAsyncOpLock)
 		, RowIndex(InRowIndex)
 		{}
 
@@ -809,15 +870,18 @@ public:
 
 	void DoTask(ENamedThreads::Type CurrentThread, const FGraphEventRef& MyCompletionGraphEvent)
 	{
-		if (!CancellationToken->ShouldCancel())
+		FScopedSharedRWLock ScopedLock(AsyncOpLock);
+		if (!CancellationToken->ShouldCancel() && TableTreeViewPtr)
 		{
 			TableTreeViewPtr->SelectNodeByTableRowIndex(RowIndex);
 		}
+		AsyncOpLock.Reset();
 	}
 
 private:
 	TSharedPtr< FTableTaskCancellationToken> CancellationToken;
 	TSharedPtr<STableTreeView> TableTreeViewPtr;
+	TSharedPtr<FRWLock> AsyncOpLock;
 	uint32 RowIndex;
 };
 
