Disk Size Tool -- Only autoexpand one level deep in tree

diff --git a/Engine/Plugins/Editor/AssetManagerEditor/Source/AssetManagerEditor/Private/Insights/Table/Widgets/STableTreeView.h b/Engine/Plugins/Editor/AssetManagerEditor/Source/AssetManagerEditor/Private/Insights/Table/Widgets/STableTreeView.h
--- a/Engine/Plugins/Editor/AssetManagerEditor/Source/AssetManagerEditor/Private/Insights/Table/Widgets/STableTreeView.h
+++ b/Engine/Plugins/Editor/AssetManagerEditor/Source/AssetManagerEditor/Private/Insights/Table/Widgets/STableTreeView.h
@@ -533,7 +533,7 @@ protected:
 	TSet<FTableTreeNodePtr> ExpandedNodes;
 
 	static constexpr int32 MaxNodesToAutoExpand = 1000;
-	static constexpr int32 MaxDepthToAutoExpand = 3;
+	static constexpr int32 MaxDepthToAutoExpand = 1;
 	static constexpr int32 MaxNodesToExpand = 1000000;
 	static constexpr int32 MaxDepthToExpand = 100;
 
