Asset Disk Size Tool -- Added support for reading ucookmeta files for additional plugin dependency and cook info.

diff --git a/Engine/Plugins/Editor/AssetManagerEditor/Source/AssetManagerEditor/AssetManagerEditor.Build.cs b/Engine/Plugins/Editor/AssetManagerEditor/Source/AssetManagerEditor/AssetManagerEditor.Build.cs
--- a/Engine/Plugins/Editor/AssetManagerEditor/Source/AssetManagerEditor/AssetManagerEditor.Build.cs
+++ b/Engine/Plugins/Editor/AssetManagerEditor/Source/AssetManagerEditor/AssetManagerEditor.Build.cs
@@ -23,6 +23,7 @@ public class AssetManagerEditor : ModuleRules
 		
 		PrivateDependencyModuleNames.AddRange(
 			new string[] {
+				"CookMetadata",
 				"Slate",
 				"SlateCore",
 				"ApplicationCore",
diff --git a/Engine/Plugins/Editor/AssetManagerEditor/Source/AssetManagerEditor/Private/AssetManagerEditorModule.cpp b/Engine/Plugins/Editor/AssetManagerEditor/Source/AssetManagerEditor/Private/AssetManagerEditorModule.cpp
--- a/Engine/Plugins/Editor/AssetManagerEditor/Source/AssetManagerEditor/Private/AssetManagerEditorModule.cpp
+++ b/Engine/Plugins/Editor/AssetManagerEditor/Source/AssetManagerEditor/Private/AssetManagerEditorModule.cpp
@@ -322,7 +322,7 @@ public:
 	virtual void GetAvailableRegistrySources(TArray<const FAssetManagerEditorRegistrySource*>& AvailableSources) override;
 	virtual const FAssetManagerEditorRegistrySource* GetCurrentRegistrySource(bool bNeedManagementData = false) override;
 	virtual void SetCurrentRegistrySource(const FString& SourceName) override;
-	virtual bool PopulateRegistrySource(FAssetManagerEditorRegistrySource* OutRegistrySource) override;
+	virtual bool PopulateRegistrySource(FAssetManagerEditorRegistrySource* OutRegistrySource, const FString* OptInFilePath = nullptr) override;
 	virtual void RefreshRegistryData() override;
 	virtual bool IsPackageInCurrentRegistrySource(FName PackageName) override;
 	virtual bool FilterAssetIdentifiersForCurrentRegistrySource(TArray<FAssetIdentifier>& AssetIdentifiers, const FAssetManagerDependencyQuery& DependencyQuery = FAssetManagerDependencyQuery::None(), bool bForwardDependency = true) override;
@@ -1348,18 +1348,8 @@ bool FAssetManagerEditorModule::GetStringValueForCustomColumn(const FAssetData&
 	}
 	else if (ColumnName == PluginName)
 	{
-		TArray<FString> ParsedPath;
-		const FString& PackageNameString = AssetData.PackageName.ToString();
-		if (PackageNameString.StartsWith(TEXT("/Game/")))
-		{
-			OutValue = TEXT("Game");
-			return true;
-		}
-		else if (PackageNameString.ParseIntoArray(ParsedPath, TEXT("/")) >= 1)
-		{
-			OutValue = ParsedPath[0];
-			return true;
-		}
+		OutValue = FPackageName::SplitPackageNameRoot(AssetData.PackageName.ToString(), nullptr);
+		return OutValue.Len() > 0;
 	}
 	else
 	{
@@ -1662,26 +1652,33 @@ void FAssetManagerEditorRegistrySource::LoadRegistryTimestamp()
 	}
 }
 
-bool FAssetManagerEditorModule::PopulateRegistrySource(FAssetManagerEditorRegistrySource* InOutRegistrySource)
+bool FAssetManagerEditorModule::PopulateRegistrySource(FAssetManagerEditorRegistrySource* InOutRegistrySource, const FString* OptInFilePath/*= nullptr*/)
 {
 	if (InOutRegistrySource->SourceName == FAssetManagerEditorRegistrySource::CustomSourceName)
 	{
-		IDesktopPlatform* DesktopPlatform = FDesktopPlatformModule::Get();
-		const void* ParentWindowWindowHandle = FSlateApplication::Get().FindBestParentWindowHandleForDialogs(nullptr);
-		const TCHAR* DevelopmentAssetRegistryFilename = GetDevelopmentAssetRegistryFilename();
-		const FText Title = LOCTEXT("LoadAssetRegistry", "Load DevelopmentAssetRegistry");
-		const FString FileTypes = FString::Printf(TEXT("%s|*.bin"), DevelopmentAssetRegistryFilename);
-
 		TArray<FString> OutFilenames;
-		DesktopPlatform->OpenFileDialog(
-			ParentWindowWindowHandle,
-			Title.ToString(),
-			TEXT(""),
-			DevelopmentAssetRegistryFilename,
-			FileTypes,
-			EFileDialogFlags::None,
-			OutFilenames
-		);
+		if (OptInFilePath != nullptr && IFileManager::Get().FileExists(**OptInFilePath))
+		{
+			OutFilenames.Add(*OptInFilePath);
+		}
+		else
+		{
+			IDesktopPlatform* DesktopPlatform = FDesktopPlatformModule::Get();
+			const void* ParentWindowWindowHandle = FSlateApplication::Get().FindBestParentWindowHandleForDialogs(nullptr);
+			const TCHAR* DevelopmentAssetRegistryFilename = GetDevelopmentAssetRegistryFilename();
+			const FText Title = LOCTEXT("LoadAssetRegistry", "Load DevelopmentAssetRegistry");
+			const FString FileTypes = FString::Printf(TEXT("%s|*.bin"), DevelopmentAssetRegistryFilename);
+
+			DesktopPlatform->OpenFileDialog(
+				ParentWindowWindowHandle,
+				Title.ToString(),
+				TEXT(""),
+				DevelopmentAssetRegistryFilename,
+				FileTypes,
+				EFileDialogFlags::None,
+				OutFilenames
+			);
+		}
 
 		if (OutFilenames.Num() == 1)
 		{
diff --git a/Engine/Plugins/Editor/AssetManagerEditor/Source/AssetManagerEditor/Private/TreeView/AssetTable.cpp b/Engine/Plugins/Editor/AssetManagerEditor/Source/AssetManagerEditor/Private/TreeView/AssetTable.cpp
--- a/Engine/Plugins/Editor/AssetManagerEditor/Source/AssetManagerEditor/Private/TreeView/AssetTable.cpp
+++ b/Engine/Plugins/Editor/AssetManagerEditor/Source/AssetManagerEditor/Private/TreeView/AssetTable.cpp
@@ -647,7 +647,16 @@ void FAssetTable::AddDefaultColumns()
 		public:
 			virtual const TOptional<FTableCellValue> GetValue(const FTableColumn& Column, const FBaseTreeNode& Node) const override
 			{
-				if (Node.Is<FAssetTreeNode>() && Node.As<FAssetTreeNode>().IsValidAsset())
+				if (Node.Is<FPluginSimpleGroupNode>() && !Node.Is<FPluginDependenciesGroupNode>())
+				{
+					// This is node represents a single plugin (it might be the plugin itself or the plugin+deps node for that plugin)
+					const FPluginSimpleGroupNode& PluginNode = Node.As<FPluginSimpleGroupNode>();
+					TSharedPtr<FTable> TablePtr = PluginNode.GetParentTable().Pin();
+					const FAssetTable& AssetTable = static_cast<const FAssetTable&>(*TablePtr);
+					const FAssetTablePluginInfo& PluginInfo = AssetTable.GetPluginInfoByIndexChecked(PluginNode.GetPluginIndex());
+					return FTableCellValue(PluginInfo.GetSize());
+				}
+				else if (Node.Is<FAssetTreeNode>() && Node.As<FAssetTreeNode>().IsValidAsset())
 				{
 					if (Node.Is<FAssetDependenciesGroupTreeNode>())
 					{
diff --git a/Engine/Plugins/Editor/AssetManagerEditor/Source/AssetManagerEditor/Private/TreeView/AssetTable.h b/Engine/Plugins/Editor/AssetManagerEditor/Source/AssetManagerEditor/Private/TreeView/AssetTable.h
--- a/Engine/Plugins/Editor/AssetManagerEditor/Source/AssetManagerEditor/Private/TreeView/AssetTable.h
+++ b/Engine/Plugins/Editor/AssetManagerEditor/Source/AssetManagerEditor/Private/TreeView/AssetTable.h
@@ -147,6 +147,7 @@ public:
 	int32 GetNumReferencers() const { return PluginReferencers.Num(); }
 	const TArray<int32>& GetReferencers() const { return PluginReferencers; }
 	int64 GetSize() const { return Size; }
+	bool GetIsRootPlugin() const { return IsRootPlugin; }
 
 	int64 GetOrComputeTotalSizeInclusiveOfDependencies(const FAssetTable& OwningTable) const;
 
@@ -159,13 +160,13 @@ private:
 	void ComputeDependencySizes(const FAssetTable& OwningTable) const;
 
 	TArray<int32> PluginDependencies;
-	TArray<int32> DiscoveredPluginDependencies;
 	TArray<int32> PluginReferencers;
 	const TCHAR* PluginName = nullptr;
 	int64 Size = 0;
 	mutable int64 InclusiveSize = -1;
 	mutable int64 UniqueDependenciesSize = -1;
 	mutable int64 SharedDependenciesSize = -1;
+	bool IsRootPlugin = false;
 };
 
 class FAssetTableRow
diff --git a/Engine/Plugins/Editor/AssetManagerEditor/Source/AssetManagerEditor/Private/TreeView/SAssetTableTreeView.cpp b/Engine/Plugins/Editor/AssetManagerEditor/Source/AssetManagerEditor/Private/TreeView/SAssetTableTreeView.cpp
--- a/Engine/Plugins/Editor/AssetManagerEditor/Source/AssetManagerEditor/Private/TreeView/SAssetTableTreeView.cpp
+++ b/Engine/Plugins/Editor/AssetManagerEditor/Source/AssetManagerEditor/Private/TreeView/SAssetTableTreeView.cpp
@@ -7,6 +7,7 @@
 #include "AssetRegistry/AssetRegistryModule.h"
 #include "Containers/Set.h"
 #include "ContentBrowserModule.h"
+#include "CookMetadata.h"
 #include "DesktopPlatformModule.h"
 #include "Editor.h"
 #include "Engine/AssetManager.h"
@@ -24,7 +25,9 @@
 #include "Insights/Table/ViewModels/TreeNodeGrouping.h"
 #include "Interfaces/IPluginManager.h"
 #include "Logging/MessageLog.h"
+#include "Misc/FileHelper.h"
 #include "Modules/ModuleManager.h"
+#include "Serialization/ArrayReader.h"
 #include "SlateOptMacros.h"
 #include "Styling/AppStyle.h"
 #include "Styling/StyleColors.h"
@@ -72,7 +75,7 @@ void SAssetTableTreeView::Construct(const FArguments& InArgs, TSharedPtr<FAssetT
 	CreateGroupings();
 	CreateSortings();
 
-	RegistrySourceTimeText->SetText(LOCTEXT("RegistrySourceTimeText_None", "No registry loaded."));
+	RegistryInfoText->SetText(LOCTEXT("RegistrySourceTimeText_None", "No registry loaded."));
 
 	RequestOpenRegistry();
 }
@@ -166,16 +169,9 @@ void SAssetTableTreeView::ConstructHeaderArea(TSharedRef<SVerticalBox> InWidgetC
 			.FillWidth(1.0f)
 			.VAlign(VAlign_Center)
 			[
-				ConstructSearchBox()
+				SAssignNew(RegistryInfoText, STextBlock)
 			]
 
-			+ SHorizontalBox::Slot()
-			.AutoWidth()
-			.Padding(4.0f, 0.0f, 0.0f, 0.0f)
-			.VAlign(VAlign_Center)
-			[
-				ConstructFilterConfiguratorButton()
-			]
 		];
 
 	InWidgetContent->AddSlot()
@@ -189,7 +185,15 @@ void SAssetTableTreeView::ConstructHeaderArea(TSharedRef<SVerticalBox> InWidgetC
 			.FillWidth(1.0f)
 			.VAlign(VAlign_Center)
 			[
-				SAssignNew(RegistrySourceTimeText, STextBlock)
+				ConstructSearchBox()
+			]
+
+			+ SHorizontalBox::Slot()
+			.AutoWidth()
+			.Padding(4.0f, 0.0f, 0.0f, 0.0f)
+			.VAlign(VAlign_Center)
+			[
+				ConstructFilterConfiguratorButton()
 			]
 		];
 
@@ -1031,27 +1035,344 @@ void SAssetTableTreeView::RequestOpenRegistry()
 
 ////////////////////////////////////////////////////////////////////////////////////////////////////
 
+bool HashRegistryFile(const FString& FilePath, uint64* HashOut)
+{
+	check(HashOut != nullptr);
+
+	bool Success = false;
+	FArrayReader SerializedAssetData;
+	if (FFileHelper::LoadFileToArray(SerializedAssetData, *FilePath))
+	{
+		*HashOut = UE::Cook::FCookMetadataState::ComputeHashOfDevelopmentAssetRegistry(MakeMemoryView(static_cast<TArray<uint8>&>(SerializedAssetData)));
+		Success = true;
+	}
+	return Success;	
+}
+
+////////////////////////////////////////////////////////////////////////////////////////////////////
+
+bool LoadCookMetadata(const FString& FilePath, UE::Cook::FCookMetadataState& OutMetadataState)
+{
+	bool Success = false;
+	FArrayReader SerializedAssetData;
+	if (FFileHelper::LoadFileToArray(SerializedAssetData, *FilePath))
+	{
+		Success = OutMetadataState.Serialize(SerializedAssetData);
+		if (!Success)
+		{
+			OutMetadataState.Reset();
+		}
+	}
+	return Success;
+}
+
+////////////////////////////////////////////////////////////////////////////////////////////////////
+
 void SAssetTableTreeView::OpenRegistry()
 {
-	RegistrySource.SourceName = FAssetManagerEditorRegistrySource::CustomSourceName;
-	if (EditorModule->PopulateRegistrySource(&RegistrySource))
+	// Ideally, we will load two files: a DevelopmentAssetRegistry.bin and a matching .ucookmeta file. However, it may happen that no ucookmeta is available
+	// We will pop a dialog that asks the user to select either a .bin or a .ucookmeta. If they pick the .bin, we will attempt to find an appropriate ucookmeta
+	// and vice versa.
+
+	IDesktopPlatform* DesktopPlatform = FDesktopPlatformModule::Get();
+	const void* ParentWindowWindowHandle = FSlateApplication::Get().FindBestParentWindowHandleForDialogs(nullptr);
+	extern const TCHAR* GetDevelopmentAssetRegistryFilename();
+	const TCHAR* DevelopmentAssetRegistryCanonicalFilename = GetDevelopmentAssetRegistryFilename();
+	const TCHAR* CookMetadataCanonicalFilename = *UE::Cook::GetCookMetadataFilename();
+	const FString CookMetadataExtension = FPaths::GetExtension(CookMetadataCanonicalFilename);
+	const FString DevelopmentAssetRegistryExtension = FPaths::GetExtension(DevelopmentAssetRegistryCanonicalFilename);;
+	const FText Title = LOCTEXT("LoadAssetRegistryOrCookMetadata", "Load DevelopmentAssetRegistry or CookMetadata");
+	const FString FileTypes = FString::Printf(TEXT("%s|*.%s|%s|*.%s"), 
+		CookMetadataCanonicalFilename, *CookMetadataExtension, DevelopmentAssetRegistryCanonicalFilename, *DevelopmentAssetRegistryExtension);
+
+	TArray<FString> OutFilenames;
+
+	bool CanceledOpenDialog = !DesktopPlatform->OpenFileDialog(
+		ParentWindowWindowHandle,
+		Title.ToString(),
+		TEXT(""),
+		UE::Cook::GetCookMetadataFilename(),
+		FileTypes,
+		EFileDialogFlags::None,
+		OutFilenames
+	);
+
+	if (CanceledOpenDialog)
 	{
-		IFileManager& FileManager = IFileManager::Get();
-		const FString RegistryFilePath = FileManager.GetFilenameOnDisk(*FileManager.ConvertToAbsolutePathForExternalAppForRead(*RegistrySource.SourceFilename));
-		RegistrySourceTimeText->SetText(FText::Format(LOCTEXT("RegistrySourceTimeText", "Loaded: {0} from {1}"), 
-			FText::FromString(RegistryFilePath),
-			FText::FromString(RegistrySource.SourceTimestamp)));
+		FooterLeftText = FooterLeftTextStoredPreOpen;
+		return;
+	}
+
+	FString RegistryFilename;
+	FString MetadataFilename;
+
+	bool FoundOtherFileByInference = false;
+
+	UE::Cook::FCookMetadataState MetadataTemporaryStorage;
+	ECheckFilesExistAndHashMatchesResult HashCheckResult = ECheckFilesExistAndHashMatchesResult::Unknown;
+
+	if (OutFilenames.Num() == 1)
+	{
+		const FString& Filename = OutFilenames[0];
+		FString FileExtension = FPaths::GetExtension(Filename);
+		bool IsRegistryFile = FileExtension.Equals(DevelopmentAssetRegistryExtension, ESearchCase::IgnoreCase);
+		bool IsMetadataFile = FileExtension.Equals(CookMetadataExtension, ESearchCase::IgnoreCase);
+
+		// Try to automatically find the corresponding file
+		const FString Directory = FPaths::GetPath(Filename);
+		const FString InitialBaseFilename = FPaths::GetBaseFilename(Filename);
+		FString NewFilenameToTry;
+		if (IsRegistryFile)
+		{
+			RegistryFilename = Filename;
+			NewFilenameToTry = Directory / InitialBaseFilename.Replace(TEXT("DevelopmentAssetRegistry"), TEXT("CookMetadata"));
+			NewFilenameToTry += TEXT(".ucookmeta");
+			if (IFileManager::Get().FileExists(*NewFilenameToTry))
+			{
+				MetadataFilename = NewFilenameToTry;
+				FoundOtherFileByInference = true;
+			}
+		}
+		else if (IsMetadataFile)
+		{
+			MetadataFilename = Filename;
+			NewFilenameToTry = Directory / InitialBaseFilename.Replace(TEXT("CookMetadata"), TEXT("DevelopmentAssetRegistry"));
+			NewFilenameToTry += TEXT(".bin");
+			if (IFileManager::Get().FileExists(*NewFilenameToTry))
+			{
+				RegistryFilename = NewFilenameToTry;
+				FoundOtherFileByInference = true;
+			}
+		}
+
+		// We think we found the other file by inference. Let's check the hashes and see  if they match. If not, we'll prompt for the matching file
+		if (FoundOtherFileByInference)
+		{
+			HashCheckResult = CheckFilesExistAndHashMatches(MetadataFilename, RegistryFilename, MetadataTemporaryStorage);
+			FoundOtherFileByInference = HashCheckResult == ECheckFilesExistAndHashMatchesResult::Okay;
+		}
+
+		if (!FoundOtherFileByInference)
+		{
+			TArray<FString> OutFollowupFilenames;
+			const FText FollowupTitle = IsRegistryFile ? LOCTEXT("LoadCookMetadata", "Load CookMetadata") : LOCTEXT("LoadAssetRegistry", "Load DevelopmentAssetRegistry");
+			const FString FollowupFileTypes = FString::Printf(TEXT("%s|*.%s"), 
+				IsRegistryFile ? CookMetadataCanonicalFilename : DevelopmentAssetRegistryCanonicalFilename, 
+				IsRegistryFile ? *CookMetadataExtension : *DevelopmentAssetRegistryExtension);
+			FString DefaultFilename = IsRegistryFile ? CookMetadataCanonicalFilename : DevelopmentAssetRegistryCanonicalFilename;
+		
+			CanceledOpenDialog = !DesktopPlatform->OpenFileDialog(
+				ParentWindowWindowHandle,
+				FollowupTitle.ToString(),
+				TEXT(""),
+				DefaultFilename,
+				FollowupFileTypes,
+				EFileDialogFlags::None,
+				OutFollowupFilenames
+			);
+
+			if (CanceledOpenDialog && !IsRegistryFile)
+			{
+				FooterLeftText = FooterLeftTextStoredPreOpen;
+				return;
+			}
+
+			if (OutFollowupFilenames.Num() == 1)
+			{
+				const FString& FollowupFilename = OutFollowupFilenames[0];
+				const FString FollowupFileExtension = FPaths::GetExtension(FollowupFilename);
+				bool FollowupIsRegistryFile = FollowupFileExtension.Equals(DevelopmentAssetRegistryExtension, ESearchCase::IgnoreCase);
+				bool FollowupIsMetadataFile = FollowupFileExtension.Equals(CookMetadataExtension, ESearchCase::IgnoreCase);
 
-		RequestRefreshAssets();
+				if (IsRegistryFile && FollowupIsMetadataFile)
+				{
+					MetadataFilename = FollowupFilename;
+				}
+				else if (IsMetadataFile && FollowupIsRegistryFile)
+				{
+					RegistryFilename = FollowupFilename;
+				}
+				else if (!IsMetadataFile && !IsRegistryFile && FollowupIsRegistryFile)
+				{
+					// Weird case where the user picked a bad metadata file (not ending in .ucookmeta) but a valid registry
+					// Just treat this as though they only picked the registry file and report a missing metadata file.
+					RegistryFilename = FollowupFilename;
+				}
+			}
+		}
+	}
+
+
+	// We didn't find the other file by inference, at least not successfully, so we need to check again using the new filename from the user
+	if (!FoundOtherFileByInference)
+	{
+		HashCheckResult = CheckFilesExistAndHashMatches(MetadataFilename, RegistryFilename, MetadataTemporaryStorage);
+	}
+
+	// We can continue loading if:
+	// (a) We have both files and the hashes match
+	// (b) We only got the development asset registry file
+	const bool MetadataFileExists = FPaths::FileExists(*MetadataFilename);
+	const bool RegistryFileExists = FPaths::FileExists(*RegistryFilename);
+	bool ShouldTryFullLoad = (HashCheckResult == ECheckFilesExistAndHashMatchesResult::Okay) || (RegistryFileExists && !MetadataFileExists);
+	
+	bool Success = false;
+	if (ShouldTryFullLoad)
+	{
+		// It's okay to stomp the SourceName since it will always be either uninitialized or CustomSourceName
+		RegistrySource.SourceName = FAssetManagerEditorRegistrySource::CustomSourceName;
+		if (EditorModule->PopulateRegistrySource(&RegistrySource, &RegistryFilename))
+		{
+			RequestRefreshAssets();
+			Success = true;
+		}
+	}
+
+	if (Success && (HashCheckResult == ECheckFilesExistAndHashMatchesResult::Okay))
+	{
+		CookMetadata = std::move(MetadataTemporaryStorage);
+		UpdateRegistryInfoTextPostLoad(HashCheckResult);
 	}
-	else if (!IsRegistrySourceValid())
+	else if (Success)
 	{
-		FooterLeftText = LOCTEXT("FooterLeftTextFmt_OpenRegistry_Failed", "No registry selected or load failed.");
-		RegistrySourceTimeText->SetText(LOCTEXT("RegistrySourceTimeText_None", "No registry loaded."));
+		UpdateRegistryInfoTextPostLoad(ECheckFilesExistAndHashMatchesResult::CookMetadataDoesNotExist);
 	}
 	else
 	{
-		FooterLeftText = FooterLeftTextStoredPreOpen;
+		ECheckFilesExistAndHashMatchesResult StatusResult = ECheckFilesExistAndHashMatchesResult::Unknown;
+
+		if (ShouldTryFullLoad)
+		{
+			StatusResult = ECheckFilesExistAndHashMatchesResult::FailedToLoadRegistry;
+		}
+		else if (CanceledOpenDialog)
+		{
+			FooterLeftText = FooterLeftTextStoredPreOpen;
+		}
+		else if (!RegistryFileExists)
+		{
+			StatusResult = ECheckFilesExistAndHashMatchesResult::RegistryDoesNotExist;
+		}
+		else
+		{
+			StatusResult = HashCheckResult;
+		}
+
+		if (!CanceledOpenDialog)
+		{
+			CookMetadata.Reset();
+			RegistrySource.ClearRegistry();
+			UpdateRegistryInfoTextPostLoad(StatusResult);
+			ClearTableAndTree();
+		}
+	}
+}
+
+////////////////////////////////////////////////////////////////////////////////////////////////////
+
+SAssetTableTreeView::ECheckFilesExistAndHashMatchesResult SAssetTableTreeView::CheckFilesExistAndHashMatches(const FString& MetadataFilename, const FString& RegistryFilename, UE::Cook::FCookMetadataState& MetadataTemporaryStorage)
+{
+	const bool MetadataFileExists = FPaths::FileExists(*MetadataFilename);
+	const bool RegistryFileExists = FPaths::FileExists(*RegistryFilename);
+	if (RegistryFileExists && MetadataFileExists)
+	{
+		uint64 RegistryFileActualHash = -1;
+		bool GotActualRegistryFileHash = HashRegistryFile(RegistryFilename, &RegistryFileActualHash);
+		if (GotActualRegistryFileHash)
+		{
+			uint64 RegistryFileStoredHash = -1;
+			if (LoadCookMetadata(MetadataFilename, MetadataTemporaryStorage))
+			{
+				bool HashesMatch = MetadataTemporaryStorage.GetAssociatedDevelopmentAssetRegistryHashPostWriteback() == RegistryFileActualHash;
+				if (!HashesMatch)
+				{
+					HashesMatch = MetadataTemporaryStorage.GetAssociatedDevelopmentAssetRegistryHash() == RegistryFileActualHash;
+				}
+				return HashesMatch ? ECheckFilesExistAndHashMatchesResult::Okay : ECheckFilesExistAndHashMatchesResult::HashesDoNotMatch;
+			}
+			else
+			{
+				return ECheckFilesExistAndHashMatchesResult::FailedToLoadCookMetadata;
+			}
+		}
+		else
+		{
+			return ECheckFilesExistAndHashMatchesResult::FailedToHashRegistry;
+		}
+	}	
+	else
+	{
+		if (!RegistryFileExists)
+		{
+			return ECheckFilesExistAndHashMatchesResult::RegistryDoesNotExist;
+		}
+		else
+		{
+			return ECheckFilesExistAndHashMatchesResult::CookMetadataDoesNotExist;
+		}
+	}
+}
+
+////////////////////////////////////////////////////////////////////////////////////////////////////
+
+void SAssetTableTreeView::UpdateRegistryInfoTextPostLoad(ECheckFilesExistAndHashMatchesResult StatusResult)
+{
+	if (StatusResult == ECheckFilesExistAndHashMatchesResult::Okay || StatusResult == ECheckFilesExistAndHashMatchesResult::CookMetadataDoesNotExist)
+	{
+		IFileManager& FileManager = IFileManager::Get();
+		const FString RegistryFilePath = FileManager.GetFilenameOnDisk(*FileManager.ConvertToAbsolutePathForExternalAppForRead(*RegistrySource.SourceFilename));
+
+		if (!CookMetadata.IsValid())
+		{
+			RegistryInfoText->SetText(FText::Format(
+				LOCTEXT("RegistrySourceInfoText_Loaded_RegistryOnly", "Loaded: {0} from {1}. No cook metadata available."),
+				FText::FromString(RegistryFilePath),
+				FText::FromString(RegistrySource.SourceTimestamp)));
+		}
+		else
+		{
+			FString HordeJobIdString = (CookMetadata.GetHordeJobId().Len() > 0) ? CookMetadata.GetHordeJobId() : TEXT("<Unavailable>");
+			RegistryInfoText->SetText(FText::Format(
+				LOCTEXT("RegistrySourceInfoText_Loaded_WithMetadata", 
+					"Loaded: {0} from build {1} on {2} with HordeJobId {3}. Platform: {4} (Sizes are {5})"),
+				FText::FromString(RegistryFilePath),
+				FText::FromString(CookMetadata.GetBuildVersion()),
+				FText::FromString(RegistrySource.SourceTimestamp),
+				FText::FromString(HordeJobIdString),
+				FText::FromString(CookMetadata.GetPlatform()),
+				CookMetadata.GetSizesPresentAsText()
+			));
+		}
+	}
+	else
+	{
+		FText InfoText;
+		switch (StatusResult)
+		{
+		case ECheckFilesExistAndHashMatchesResult::RegistryDoesNotExist:
+			InfoText = LOCTEXT("RegistrySourceInfoText_OpenRegistry_RegistryDoesNotExist", "Selected asset registry does not exist or is invalid.");
+			break;
+		case ECheckFilesExistAndHashMatchesResult::FailedToLoadCookMetadata:
+			InfoText = LOCTEXT("RegistrySourceInfoText_OpenRegistry_CouldNotLoadMetadata", "Unable to load cook metadata.");
+			break;
+		case ECheckFilesExistAndHashMatchesResult::FailedToHashRegistry:
+			InfoText = LOCTEXT("RegistrySourceInfoText_OpenRegistry_CouldNotHashAssetRegistry", "Unable to load registry to hash it.");
+			break;
+		case ECheckFilesExistAndHashMatchesResult::HashesDoNotMatch:
+			InfoText = LOCTEXT("RegistrySourceInfoText_OpenRegistry_HashesDoNotMatch", "Hash of asset registry does not match either hash provided by the cook metadata file.");
+			break;
+		case ECheckFilesExistAndHashMatchesResult::FailedToLoadRegistry:
+			InfoText = LOCTEXT("RegistrySourceInfoText_OpenRegistry_DevARLoadFailed", "Selected asset registry could not be loaded/parsed.");
+			break;
+		case ECheckFilesExistAndHashMatchesResult::Unknown:
+			InfoText = LOCTEXT("RegistrySourceInfoText_OpenRegistry_Failed", "No registry selected or load failed. Reason unknown.");
+			break;
+		default:
+			ensureAlwaysMsgf(false, TEXT("Unexpected status result in UpdateRegistryInfoTextPostLoad"));
+			InfoText = LOCTEXT("RegistrySourceInfoText_OpenRegistry_UnknownError", "Internal error.");
+			break;
+		};
+		RegistryInfoText->SetText(InfoText);
 	}
 }
 
@@ -1393,6 +1714,7 @@ void SAssetTableTreeView::PopulateAssetTableRow(FAssetTableRow& OutRow, const FA
 // Refresh list of assets for the tree view
 void SAssetTableTreeView::RefreshAssets()
 {
+	CancelCurrentAsyncOp();
 	TSharedPtr<FAssetTable> AssetTable = GetAssetTable();
 	if (!AssetTable.IsValid())
 	{
@@ -1401,21 +1723,13 @@ void SAssetTableTreeView::RefreshAssets()
 
 	UE_LOG(LogInsights, Log, TEXT("[AssetTree] Build asset table..."));
 
-	CancelCurrentAsyncOp();
-
 	UE::Insights::FStopwatch Stopwatch;
 	Stopwatch.Start();
 
-	// Clears all tree nodes (that references the previous assets).
-	AssetTable->SetVisibleAssetCount(0);
-	RebuildTree(true);
-
-	// Now is safe to clear the previous assets.
-	AssetTable->ClearAllData();
+	ClearTableAndTree();
 
 	TMap<FString, int64> PluginToSizeMap;
 
-	typedef TSet<const TCHAR*, TStringPointerSetKeyFuncs_DEPRECATED<const TCHAR*>> DeprecatedTCharSetType;
 	TMap<const TCHAR*, DeprecatedTCharSetType, FDefaultSetAllocator, TStringPointerMapKeyFuncs_DEPRECATED<const TCHAR*, DeprecatedTCharSetType>> DiscoveredPluginDependencyEdges;
 
 	if (IsRegistrySourceValid())
@@ -1522,6 +1836,7 @@ void SAssetTableTreeView::RefreshAssets()
 					{
 						UE_LOG(LogInsights, Warning, TEXT("Failed to find asset %s in package %s, source asset index %d. Asset registry loading was %s"), *SourceAssetData.AssetName.ToString(), *SourceAssetData.PackageName.ToString(), SourceAssetIndex, AssetRegistry->IsLoadingAssets() ? TEXT("INCOMPLETE") : TEXT("complete"));
 					}
+
 					if (ensure(RowIndex != nullptr))
 					{
 						const TCHAR* CurrentPlugin = AssetTable->GetAssetChecked(*RowIndex).GetPluginName();
@@ -1538,7 +1853,6 @@ void SAssetTableTreeView::RefreshAssets()
 							}
 						}
 
-
 						DeprecatedTCharSetType& DiscoveredPluginDependencyList = DiscoveredPluginDependencyEdges.FindOrAdd(CurrentPlugin);
 						for (const TCHAR* DependencyPlugin : DependencyPlugins)
 						{
@@ -1567,20 +1881,42 @@ void SAssetTableTreeView::RefreshAssets()
 		}
 	}
 
-	// Setup plugin infos and dependencies. This will eventually be replaced by data from the asset registry
-	for (TPair<FString, int64>& PluginEntry : PluginToSizeMap)
+	if (CookMetadata.IsValid())
 	{
-		const TCHAR* StoredPluginName = AssetTable->StoreStr(PluginEntry.Key);
-		FAssetTablePluginInfo& PluginInfo = AssetTable->GetOrCreatePluginInfo(StoredPluginName);
-		PluginInfo.Size = PluginEntry.Value;
-		if (TSharedPtr<IPlugin> Plugin = IPluginManager::Get().FindPlugin(PluginEntry.Key))
+		const UE::Cook::FCookMetadataPluginHierarchy& PluginHierarchy = CookMetadata.GetPluginHierarchy();
+		for (const UE::Cook::FCookMetadataPluginEntry& PluginEntry : PluginHierarchy.PluginsEnabledAtCook)
 		{
-			TArray<FPluginReferenceDescriptor> PluginReferences;
-			IPluginManager::Get().GetPluginDependencies(PluginEntry.Key, PluginReferences);
+			const TCHAR* StoredPluginName = AssetTable->StoreStr(PluginEntry.Name);
+			FAssetTablePluginInfo& PluginInfo = AssetTable->GetOrCreatePluginInfo(StoredPluginName);
+			if (CookMetadata.GetSizesPresent() != UE::Cook::ECookMetadataSizesPresent::NotPresent)
+			{
+				PluginInfo.Size = PluginEntry.ExclusiveSizes[UE::Cook::EPluginSizeTypes::Installed];
+			}
+			else if (const int64* SizePtr = PluginToSizeMap.Find(PluginEntry.Name))
+			{
+				PluginInfo.Size = *SizePtr;
+			}
+
+			///
+			// Temp validation
+			///
+			if (const int64* SizePtr = PluginToSizeMap.Find(PluginEntry.Name))
+			{
+				if (CookMetadata.GetSizesPresent() != UE::Cook::ECookMetadataSizesPresent::NotPresent)
+				{
+					if (FMath::Abs((*SizePtr - PluginInfo.Size)) > static_cast<int64>(0.05 * FMath::Max(*SizePtr, PluginInfo.Size)))
+					{
+						UE_LOG(LogInsights, Warning, TEXT("Plugin %s found with ucookmetadata and asset calculation size delta > 5%%. Metadata size: %lld Calculated size: %lld"), 
+							StoredPluginName, PluginInfo.Size, *SizePtr);
+					}
+				}
+			}
+
 			const int32 PluginIndex = AssetTable->GetIndexForPlugin(StoredPluginName);
-			for (const FPluginReferenceDescriptor& ReferenceDescriptor : PluginReferences)
+			for (uint16 DependencyIndexInMetadata = PluginEntry.DependencyIndexStart; DependencyIndexInMetadata < PluginEntry.DependencyIndexEnd; DependencyIndexInMetadata++)
 			{
-				const TCHAR* StoredReferencePluginName = AssetTable->StoreStr(ReferenceDescriptor.Name);
+				const UE::Cook::FCookMetadataPluginEntry& DependentPluginEntry = PluginHierarchy.PluginsEnabledAtCook[PluginHierarchy.PluginDependencies[DependencyIndexInMetadata]];
+				const TCHAR* StoredReferencePluginName = AssetTable->StoreStr(DependentPluginEntry.Name);
 				int32 DependencyIndex = AssetTable->GetIndexForPlugin(StoredReferencePluginName);
 				if (DependencyIndex == -1)
 				{
@@ -1593,52 +1929,59 @@ void SAssetTableTreeView::RefreshAssets()
 				AssetTable->GetPluginInfoByIndex(DependencyIndex).PluginReferencers.AddUnique(AssetTable->GetIndexForPlugin(StoredPluginName));
 			}
 		}
-		else
+
+		for (uint16 RootPluginIndex : PluginHierarchy.RootPlugins)
 		{
-			UE_LOG(LogInsights, Warning, TEXT("Could not find plugin %s in plugin manager."), PluginInfo.PluginName);
+			const UE::Cook::FCookMetadataPluginEntry& PluginEntry = PluginHierarchy.PluginsEnabledAtCook[RootPluginIndex];
+			const TCHAR* StoredPluginName = AssetTable->StoreStr(PluginEntry.Name);
+			FAssetTablePluginInfo& PluginInfo = AssetTable->GetOrCreatePluginInfo(StoredPluginName);
+			UE_LOG(LogInsights, Warning, TEXT("Found root plugin %s"), StoredPluginName);
+			PluginInfo.IsRootPlugin = true;
 		}
 	}
-
-	// Discovered version
-	for (const TPair<const TCHAR*, DeprecatedTCharSetType>& DiscoveredDependencies : DiscoveredPluginDependencyEdges)
+	else
 	{
-		const TCHAR* StoredPluginName = AssetTable->StoreStr(DiscoveredDependencies.Key);
-		FAssetTablePluginInfo& PluginInfo = AssetTable->GetOrCreatePluginInfo(StoredPluginName);
-
-		for (const TCHAR* DependencyName : DiscoveredDependencies.Value)
-		{
-			int32 DependencyIndex = AssetTable->GetIndexForPlugin(DependencyName);
-			ensureAlways(DependencyIndex != -1); // These were created above.
-			PluginInfo.DiscoveredPluginDependencies.AddUnique(DependencyIndex);
+		UE_LOG(LogInsights, Warning, TEXT("CookMetadata not available, deriving plugin dependencies and sizes from loaded asset registry instead."));
+		// Setup plugin infos and dependencies. This will eventually be replaced by data from the asset registry
+		for (TPair<FString, int64>& PluginEntry : PluginToSizeMap)
+		{
+			const TCHAR* StoredPluginName = AssetTable->StoreStr(PluginEntry.Key);
+			FAssetTablePluginInfo& PluginInfo = AssetTable->GetOrCreatePluginInfo(StoredPluginName);
+			PluginInfo.Size = PluginEntry.Value;
+			if (TSharedPtr<IPlugin> Plugin = IPluginManager::Get().FindPlugin(PluginEntry.Key))
+			{
+				TArray<FPluginReferenceDescriptor> PluginReferences;
+				IPluginManager::Get().GetPluginDependencies(PluginEntry.Key, PluginReferences);
+				const int32 PluginIndex = AssetTable->GetIndexForPlugin(StoredPluginName);
+				for (const FPluginReferenceDescriptor& ReferenceDescriptor : PluginReferences)
+				{
+					const TCHAR* StoredReferencePluginName = AssetTable->StoreStr(ReferenceDescriptor.Name);
+					int32 DependencyIndex = AssetTable->GetIndexForPlugin(StoredReferencePluginName);
+					if (DependencyIndex == -1)
+					{
+						DependencyIndex = AssetTable->GetNumPlugins();
+						AssetTable->GetOrCreatePluginInfo(StoredReferencePluginName);
+					}
+					// Note that we can't use PluginInfo directly because the above code could have grown the array 
+					// and that would invalidate the reference
+					AssetTable->GetOrCreatePluginInfo(StoredPluginName).PluginDependencies.AddUnique(DependencyIndex);
+					AssetTable->GetPluginInfoByIndex(DependencyIndex).PluginReferencers.AddUnique(AssetTable->GetIndexForPlugin(StoredPluginName));
+				}
+			}
+			else
+			{
+				UE_LOG(LogInsights, Warning, TEXT("Could not find plugin %s in plugin manager."), PluginInfo.PluginName);
+			}
 		}
 	}
 
-	// TODO: Once we have the dependency data from the uplugin files in the cook we can 
-	// compare them to the discovered dependencies to see if there are plugins we're dependending on
-	// that, perhaps, we shouldn't be. That said, we shouldn't storer the discovered dependencies on the 
-	// plugin infos at that point.
-
-	//// Analyze difference
-	//UE_LOG(LogInsights, Warning, TEXT("Registry Deps-------"));
-	//for (int32 PluginIndex = 0; PluginIndex < AssetTable->GetNumPlugins(); PluginIndex++)
-	//{
-	//	const FAssetTablePluginInfo& PluginInfo = AssetTable->GetPluginInfoByIndex(PluginIndex);
-	//	UE_LOG(LogInsights, Warning, TEXT("Plugin: %s"), PluginInfo.PluginName);
-	//	for (int32 DependencyIndex : PluginInfo.PluginDependencies)
-	//	{
-	//		UE_LOG(LogInsights, Warning, TEXT("\t%s"), AssetTable->GetPluginInfoByIndex(DependencyIndex).PluginName);
-	//	}
-	//}
-	//UE_LOG(LogInsights, Warning, TEXT("Discovered Deps-------"));
-	//for (int32 PluginIndex = 0; PluginIndex < AssetTable->GetNumPlugins(); PluginIndex++)
-	//{
-	//	const FAssetTablePluginInfo& PluginInfo = AssetTable->GetPluginInfoByIndex(PluginIndex);
-	//	UE_LOG(LogInsights, Warning, TEXT("Plugin: %s"), PluginInfo.PluginName);
-	//	for (int32 DependencyIndex : PluginInfo.DiscoveredPluginDependencies)
-	//	{
-	//		UE_LOG(LogInsights, Warning, TEXT("\t%s"), AssetTable->GetPluginInfoByIndex(DependencyIndex).PluginName);
-	//	}
-	//}
+
+	// Disabling this for now as there are too many differences at the moment. Keeping the code because eventually
+	// we probably do want to use it to discover spurious dependencies or dependencies on code+content where content should
+	// get separated from code
+	// 
+	//DumpDifferencesBetweenDiscoveredDataAndLoadedMetadata(DiscoveredPluginDependencyEdges, AssetTable);
+
 
 	Stopwatch.Stop();
 	const double TotalTime = Stopwatch.GetAccumulatedTime();
@@ -1654,6 +1997,125 @@ void SAssetTableTreeView::RefreshAssets()
 	RequestRebuildTree();
 }
 
+void SAssetTableTreeView::ClearTableAndTree()
+{
+	TSharedPtr<FAssetTable> AssetTable = GetAssetTable();
+
+	// Clears all tree nodes (that references the previous assets).
+	AssetTable->SetVisibleAssetCount(0);
+	RebuildTree(true);
+
+	// Now is safe to clear the previous assets.
+	AssetTable->ClearAllData();
+}
+
+void SAssetTableTreeView::DumpDifferencesBetweenDiscoveredDataAndLoadedMetadata(TMap<const TCHAR*, DeprecatedTCharSetType, FDefaultSetAllocator, TStringPointerMapKeyFuncs_DEPRECATED<const TCHAR*, DeprecatedTCharSetType>>& DiscoveredPluginDependencyEdges) const
+{
+	TSharedPtr<FAssetTable> AssetTable = GetAssetTable();
+
+	if (CookMetadata.IsValid())
+	{
+		// Report discrepancies
+		TArray<TPair<FString, FString>> DependenciesInMetadataNotAssets;
+		TArray<TPair<FString, FString>> DependenciesInAssetsNotMetadata;
+
+		TMap<FString, int32> PluginNameToIndexInHierarchyMetadata;
+		for (int32 IndexInHierarchyMetadata = 0; IndexInHierarchyMetadata < CookMetadata.GetPluginHierarchy().PluginsEnabledAtCook.Num(); IndexInHierarchyMetadata++)
+		{
+			PluginNameToIndexInHierarchyMetadata.Add(CookMetadata.GetPluginHierarchy().PluginsEnabledAtCook[IndexInHierarchyMetadata].Name, IndexInHierarchyMetadata);
+		}
+
+		// Search through our discovered dependencies and ensure that we have a matching dependency in the metadata
+		for (const TPair<const TCHAR*, DeprecatedTCharSetType>& DiscoveredDependencies : DiscoveredPluginDependencyEdges)
+		{
+			int32* IndexInMetadata = PluginNameToIndexInHierarchyMetadata.Find(DiscoveredDependencies.Key);
+			if (IndexInMetadata == nullptr)
+			{
+				UE_LOG(LogInsights, Warning, TEXT("Unable to find plugin %s, known from asset traversal, in cook metadata."), DiscoveredDependencies.Key);
+			}
+			else
+			{
+				for (const TCHAR* DiscoveredDependency : DiscoveredDependencies.Value)
+				{
+					if (int32* IndexOfDependencyInMetadata = PluginNameToIndexInHierarchyMetadata.Find(DiscoveredDependency))
+					{
+						// Now make sure it's actually in the list of dependencies
+						bool FoundDependency = false;
+						int32 StartIndex = CookMetadata.GetPluginHierarchy().PluginsEnabledAtCook[*IndexInMetadata].DependencyIndexStart;
+						int32 EndIndex = CookMetadata.GetPluginHierarchy().PluginsEnabledAtCook[*IndexInMetadata].DependencyIndexEnd;
+						for (int32 MetadataDependencyIndex = StartIndex; MetadataDependencyIndex < EndIndex; MetadataDependencyIndex++)
+						{
+							if (*IndexOfDependencyInMetadata == MetadataDependencyIndex)
+							{
+								FoundDependency = true;
+								break;
+							}
+						}
+
+						if (!FoundDependency)
+						{
+							DependenciesInAssetsNotMetadata.Add(TPair<FString, FString>(DiscoveredDependencies.Key, DiscoveredDependency));
+						}
+					}
+					else
+					{
+						UE_LOG(LogInsights, Warning, TEXT("Unable to find plugin %s, known from asset traversal, in cook metadata."), DiscoveredDependency);
+					}
+				}
+
+			}
+		}
+
+		// Search through the metadata and ensure that we've found a corresponding asset dependency
+		for (int32 PluginIndex = 0; PluginIndex < AssetTable->GetNumPlugins(); PluginIndex++)
+		{
+			int32* IndexInMetadata = PluginNameToIndexInHierarchyMetadata.Find(AssetTable->GetPluginInfoByIndex(PluginIndex).GetName());
+			const FAssetTablePluginInfo& PluginInfoInTable = AssetTable->GetPluginInfoByIndex(PluginIndex);
+			if (IndexInMetadata == nullptr)
+			{
+				UE_LOG(LogInsights, Warning, TEXT("Unable to find plugin %s in metadata."), PluginInfoInTable.GetName());
+			}
+			else
+			{
+				const TCHAR* PluginName = AssetTable->GetPluginInfoByIndex(PluginIndex).GetName();
+				if (DeprecatedTCharSetType* DependencySet = DiscoveredPluginDependencyEdges.Find(PluginName))
+				{
+					for (int32 DependencyIndex : PluginInfoInTable.GetDependencies())
+					{
+						const TCHAR* DependencyName = AssetTable->GetPluginInfoByIndex(DependencyIndex).GetName();
+						if (!DependencySet->Contains(DependencyName))
+						{
+							DependenciesInMetadataNotAssets.Add(TPair<FString, FString>(PluginName, DependencyName));
+						}
+					}
+				}
+				else
+				{
+					UE_LOG(LogInsights, Warning, TEXT("Unable to find plugin %s, known from metadata, in asset traversal data"), PluginInfoInTable.GetName());
+				}
+			}
+		}
+
+		if (DependenciesInAssetsNotMetadata.Num() > 0)
+		{
+			UE_LOG(LogInsights, Warning, TEXT("Dependencies in assets but not in cook metadata:"));
+			for (const TPair<FString, FString>& Edge : DependenciesInAssetsNotMetadata)
+			{
+				UE_LOG(LogInsights, Warning, TEXT("%s-->%s"), *Edge.Key, *Edge.Value);
+			}
+		}
+
+		if (DependenciesInMetadataNotAssets.Num() > 0)
+		{
+			UE_LOG(LogInsights, Warning, TEXT("Dependencies in cook metadata but not in assets:"));
+			for (const TPair<FString, FString>& Edge : DependenciesInMetadataNotAssets)
+			{
+				UE_LOG(LogInsights, Warning, TEXT("%s-->%s"), *Edge.Key, *Edge.Value);
+			}
+		}
+	}
+}
+
 ////////////////////////////////////////////////////////////////////////////////////////////////////
 
 void SAssetTableTreeView::RequestRebuildTree()
@@ -1672,14 +2134,14 @@ void SAssetTableTreeView::RebuildTree(bool bResync)
 		return;
 	}
 
+	CancelCurrentAsyncOp();
+
 	UE::Insights::FStopwatch Stopwatch;
 	Stopwatch.Start();
 
 	UE::Insights::FStopwatch SyncStopwatch;
 	SyncStopwatch.Start();
 
-	CancelCurrentAsyncOp();
-
 	const int32 PreviousNodeCount = TableRowNodes.Num();
 	TableRowNodes.Empty();
 
diff --git a/Engine/Plugins/Editor/AssetManagerEditor/Source/AssetManagerEditor/Private/TreeView/SAssetTableTreeView.h b/Engine/Plugins/Editor/AssetManagerEditor/Source/AssetManagerEditor/Private/TreeView/SAssetTableTreeView.h
--- a/Engine/Plugins/Editor/AssetManagerEditor/Source/AssetManagerEditor/Private/TreeView/SAssetTableTreeView.h
+++ b/Engine/Plugins/Editor/AssetManagerEditor/Source/AssetManagerEditor/Private/TreeView/SAssetTableTreeView.h
@@ -3,6 +3,7 @@
 #pragma once
 
 #include "CoreMinimal.h"
+#include "CookMetadata.h"
 #include "Delegates/DelegateCombinations.h"
 #include "Widgets/DeclarativeSyntaxSupport.h"
 #include "AssetManagerEditorModule.h"
@@ -73,12 +74,28 @@ public:
 	TSharedPtr<FAssetTreeNode> GetSingleSelectedAssetNode() const { return SelectedAssetNode; }
 
 private:
+
+	enum class ECheckFilesExistAndHashMatchesResult : uint8
+	{
+		Okay,
+		RegistryDoesNotExist,
+		CookMetadataDoesNotExist,
+		FailedToHashRegistry,
+		FailedToLoadCookMetadata,
+		FailedToLoadRegistry,
+		HashesDoNotMatch,
+		Unknown
+	};
+
 	virtual void InternalCreateGroupings() override;
 
 	virtual void ExtendMenu(FMenuBuilder& MenuBuilder) override;
 
 	void RequestOpenRegistry();
 	void OpenRegistry();
+
+	ECheckFilesExistAndHashMatchesResult CheckFilesExistAndHashMatches(const FString& MetadataFilename, const FString& RegistryFilename, UE::Cook::FCookMetadataState& MetadataTemporaryStorage);
+
 	FText GetOpenRegistryToolTipText() const;
 	bool CanChangeRegistry() const;
 	FReply OnClickedOpenRegistry();
@@ -109,8 +126,15 @@ private:
 	void RequestRefreshAssets();
 	void RefreshAssets();
 
+	void ClearTableAndTree();
+
+	typedef TSet<const TCHAR*, TStringPointerSetKeyFuncs_DEPRECATED<const TCHAR*>> DeprecatedTCharSetType;
+	void DumpDifferencesBetweenDiscoveredDataAndLoadedMetadata(TMap<const TCHAR*, DeprecatedTCharSetType, FDefaultSetAllocator, TStringPointerMapKeyFuncs_DEPRECATED<const TCHAR*, DeprecatedTCharSetType>>& DiscoveredPluginDependencyEdges) const;
+
 	void RequestRebuildTree();
 
+	void UpdateRegistryInfoTextPostLoad(ECheckFilesExistAndHashMatchesResult StatusResult);
+
 private:
 	bool bNeedsToOpenRegistry = false;
 	bool bNeedsToRefreshAssets = false;
@@ -135,10 +159,12 @@ private:
 	* Text block where we put the timestamp for the asset registry so the user knows
 	* if they are looking at super out of date data.
 	*/
-	TSharedPtr<STextBlock> RegistrySourceTimeText;
+	TSharedPtr<STextBlock> RegistryInfoText;
 
 	FAssetManagerEditorRegistrySource RegistrySource;
 
+	UE::Cook::FCookMetadataState CookMetadata;
+
 	/** Cached interfaces */
 	class IAssetRegistry* AssetRegistry;
 	class UAssetManager* AssetManager;
diff --git a/Engine/Plugins/Editor/AssetManagerEditor/Source/AssetManagerEditor/Public/AssetManagerEditorModule.h b/Engine/Plugins/Editor/AssetManagerEditor/Source/AssetManagerEditor/Public/AssetManagerEditorModule.h
--- a/Engine/Plugins/Editor/AssetManagerEditor/Source/AssetManagerEditor/Public/AssetManagerEditorModule.h
+++ b/Engine/Plugins/Editor/AssetManagerEditor/Source/AssetManagerEditor/Public/AssetManagerEditorModule.h
@@ -406,8 +406,8 @@ public:
 	/** Sets the current registry source, this loads the asset registry state if needed and may spawn a file load dialog for custom */
 	virtual void SetCurrentRegistrySource(const FString& SourceName) = 0;
 
-	/** Sets up a registry source. If InOutRegistrySource->SourceName is CustomSourceName, a dialog will be presented to select the asset registry to load */
-	virtual bool PopulateRegistrySource(FAssetManagerEditorRegistrySource* InOutRegistrySource) = 0;
+	/** Sets up a registry source. If InOutRegistrySource->SourceName is CustomSourceName, a dialog will be presented to select the asset registry to load. If OptInFilePath is provided, it will be used instead of opening a dialog. */
+	virtual bool PopulateRegistrySource(FAssetManagerEditorRegistrySource* InOutRegistrySource, const FString* OptInFilePath = nullptr) = 0;
 
 	/** Refreshes the management dictionary and all sources */
 	virtual void RefreshRegistryData() = 0;
diff --git a/Engine/Source/Runtime/AssetRegistry/Private/AssetRegistry.cpp b/Engine/Source/Runtime/AssetRegistry/Private/AssetRegistry.cpp
--- a/Engine/Source/Runtime/AssetRegistry/Private/AssetRegistry.cpp
+++ b/Engine/Source/Runtime/AssetRegistry/Private/AssetRegistry.cpp
@@ -631,7 +631,7 @@ EChunkProgressReportingType::Type GetChunkAvailabilityProgressType(EAssetAvailab
 	return ChunkReportType;
 }
 
-const TCHAR* GetDevelopmentAssetRegistryFilename()
+ASSETREGISTRY_API const TCHAR* GetDevelopmentAssetRegistryFilename()
 {
 	return TEXT("DevelopmentAssetRegistry.bin");
 }
