Disk Size Tool -- For Plugin Dependencies grouping, added Root Plugin and Non-Root folders to categories plugins.

diff --git a/Engine/Plugins/Editor/AssetManagerEditor/Source/AssetManagerEditor/Private/TreeView/AssetDependencyGrouping.cpp b/Engine/Plugins/Editor/AssetManagerEditor/Source/AssetManagerEditor/Private/TreeView/AssetDependencyGrouping.cpp
--- a/Engine/Plugins/Editor/AssetManagerEditor/Source/AssetManagerEditor/Private/TreeView/AssetDependencyGrouping.cpp
+++ b/Engine/Plugins/Editor/AssetManagerEditor/Source/AssetManagerEditor/Private/TreeView/AssetDependencyGrouping.cpp
@@ -4,6 +4,7 @@
 #include "AssetTreeNode.h"
 #include "AssetTable.h"
 #include "Insights/Common/AsyncOperationProgress.h"
+#include "Insights/Common/Log.h"
 
 #define LOCTEXT_NAMESPACE "FAssetDependencyGrouping"
 
@@ -150,6 +151,40 @@ void FPluginDependencyGrouping::GroupNodes(const TArray<UE::Insights::FTableTree
 
 	TMap<int32, FPluginSimpleGroupNode*> PluginIndexToGroupNodeMap;
 
+
+	// If we have any root plugins then we'll want to ensure we always show those and we divide root from non-root
+	// However, a typical UE project may not make use of this concept so if there are no designated root plugins
+	// according to the loaded ucookmeta then only show plugins that actually have assets according to the incoming node list
+	TSharedPtr<FAssetTreeNode> RootPlugins = MakeShared<FAssetTreeNode>(FName("Root Plugins"), AssetTable);
+	TSharedPtr<FAssetTreeNode> NonRootPlugins = MakeShared<FAssetTreeNode>(FName("Non-Root Plugins"), AssetTable);
+
+	bool bFoundAnyRootPlugins = false;
+	for (int32 PluginIndex = 0; PluginIndex < AssetTable->GetNumPlugins(); PluginIndex++)
+	{
+		const FAssetTablePluginInfo& PluginInfo = AssetTable->GetPluginInfoByIndex(PluginIndex);
+
+		if (PluginInfo.GetIsRootPlugin())
+		{
+			FName PluginAndDependenciesGroupName = AssetTable->GetNameForPlugin(PluginIndex);
+			TSharedPtr<FPluginAndDependenciesGroupNode> PluginAndDependenciesGroup = MakeShared<FPluginAndDependenciesGroupNode>(PluginAndDependenciesGroupName, AssetTable, PluginIndex);
+			FPluginSimpleGroupNode* PluginGroup = PluginAndDependenciesGroup->CreateChildren().Get();
+			PluginIndexToGroupNodeMap.Add(PluginIndex, PluginGroup);
+			RootPlugins->AddChildAndSetParent(PluginAndDependenciesGroup);
+			bFoundAnyRootPlugins = true;
+		}
+	}
+
+	if (bFoundAnyRootPlugins)
+	{
+		ParentGroup.AddChildAndSetParent(RootPlugins);
+		ParentGroup.AddChildAndSetParent(NonRootPlugins);
+	}
+	else
+	{
+		RootPlugins.Reset();
+		NonRootPlugins.Reset();
+	}
+
 	for (FTableTreeNodePtr NodePtr : Nodes)
 	{
 		if (InAsyncOperationProgress.ShouldCancelAsyncOp())
@@ -241,10 +276,23 @@ void FPluginDependencyGrouping::GroupNodes(const TArray<UE::Insights::FTableTree
 			// Create the Plugin Self+Dependencies group node and add the current asset to the Self group.
 			FName PluginAndDependenciesGroupName = AssetTable->GetNameForPlugin(PluginIndex);
 			TSharedPtr<FPluginAndDependenciesGroupNode> PluginAndDependenciesGroup = MakeShared<FPluginAndDependenciesGroupNode>(PluginAndDependenciesGroupName, AssetTable, PluginIndex);
-			ParentGroup.AddChildAndSetParent(PluginAndDependenciesGroup);
 			PluginGroup = PluginAndDependenciesGroup->CreateChildren().Get();
 			PluginGroup->AddChildAndSetParent(NodePtr);
 			PluginIndexToGroupNodeMap.Add(PluginIndex, PluginGroup);
+
+			// A root plugin never has assets (by definition), so don't
+			if (NonRootPlugins.IsValid())
+			{
+				if (PluginInfo.GetIsRootPlugin())
+				{
+					UE_LOG(LogInsights, Error, TEXT("Plugin %s contains assets but is marked as a root plugin."), PluginName);
+				}
+				NonRootPlugins->AddChildAndSetParent(PluginAndDependenciesGroup);
+			}
+			else
+			{
+				ParentGroup.AddChildAndSetParent(PluginAndDependenciesGroup);
+			}
 		}
 	}
 }
