Add different size types to development asset registry.
Put the tag names in a global place to prevent string copy pasta.

diff --git a/Engine/Plugins/Editor/AssetManagerEditor/Source/AssetManagerEditor/Private/TreeView/SAssetTableTreeView.cpp b/Engine/Plugins/Editor/AssetManagerEditor/Source/AssetManagerEditor/Private/TreeView/SAssetTableTreeView.cpp
--- a/Engine/Plugins/Editor/AssetManagerEditor/Source/AssetManagerEditor/Private/TreeView/SAssetTableTreeView.cpp
+++ b/Engine/Plugins/Editor/AssetManagerEditor/Source/AssetManagerEditor/Private/TreeView/SAssetTableTreeView.cpp
@@ -1958,7 +1958,7 @@ void SAssetTableTreeView::RefreshAssets()
 			const TCHAR* StoredPluginName = AssetTable->StoreStr(PluginEntry.Name);
 			FAssetTablePluginInfo& PluginInfo = AssetTable->GetOrCreatePluginInfo(StoredPluginName);
 			UE_LOG(LogInsights, Warning, TEXT("Found root plugin %s"), StoredPluginName);
-			PluginInfo.IsRootPlugin = true;
+			PluginInfo.bIsRootPlugin = true;
 		}
 	}
 	else
