Build fix. Disk size tool. GetIsRootPlugin() --> IsRootPlugin()

diff --git a/Engine/Plugins/Editor/AssetManagerEditor/Source/AssetManagerEditor/Private/TreeView/AssetDependencyGrouping.cpp b/Engine/Plugins/Editor/AssetManagerEditor/Source/AssetManagerEditor/Private/TreeView/AssetDependencyGrouping.cpp
--- a/Engine/Plugins/Editor/AssetManagerEditor/Source/AssetManagerEditor/Private/TreeView/AssetDependencyGrouping.cpp
+++ b/Engine/Plugins/Editor/AssetManagerEditor/Source/AssetManagerEditor/Private/TreeView/AssetDependencyGrouping.cpp
@@ -163,7 +163,7 @@ void FPluginDependencyGrouping::GroupNodes(const TArray<UE::Insights::FTableTree
 	{
 		const FAssetTablePluginInfo& PluginInfo = AssetTable->GetPluginInfoByIndex(PluginIndex);
 
-		if (PluginInfo.GetIsRootPlugin())
+		if (PluginInfo.IsRootPlugin())
 		{
 			FName PluginAndDependenciesGroupName = AssetTable->GetNameForPlugin(PluginIndex);
 			TSharedPtr<FPluginAndDependenciesGroupNode> PluginAndDependenciesGroup = MakeShared<FPluginAndDependenciesGroupNode>(PluginAndDependenciesGroupName, AssetTable, PluginIndex);
@@ -283,7 +283,7 @@ void FPluginDependencyGrouping::GroupNodes(const TArray<UE::Insights::FTableTree
 			// A root plugin never has assets (by definition), so don't
 			if (NonRootPlugins.IsValid())
 			{
-				if (PluginInfo.GetIsRootPlugin())
+				if (PluginInfo.IsRootPlugin())
 				{
 					UE_LOG(LogInsights, Error, TEXT("Plugin %s contains assets but is marked as a root plugin."), PluginName);
 				}
diff --git a/Engine/Plugins/Editor/AssetManagerEditor/Source/AssetManagerEditor/Private/TreeView/AssetTable.h b/Engine/Plugins/Editor/AssetManagerEditor/Source/AssetManagerEditor/Private/TreeView/AssetTable.h
--- a/Engine/Plugins/Editor/AssetManagerEditor/Source/AssetManagerEditor/Private/TreeView/AssetTable.h
+++ b/Engine/Plugins/Editor/AssetManagerEditor/Source/AssetManagerEditor/Private/TreeView/AssetTable.h
@@ -147,7 +147,7 @@ public:
 	int32 GetNumReferencers() const { return PluginReferencers.Num(); }
 	const TArray<int32>& GetReferencers() const { return PluginReferencers; }
 	int64 GetSize() const { return Size; }
-	bool GetIsRootPlugin() const { return IsRootPlugin; }
+	bool IsRootPlugin() const { return bIsRootPlugin; }
 
 	int64 GetOrComputeTotalSizeInclusiveOfDependencies(const FAssetTable& OwningTable) const;
 
@@ -168,7 +168,7 @@ private:
 	mutable int64 InclusiveSize = -1;
 	mutable int64 UniqueDependenciesSize = -1;
 	mutable int64 SharedDependenciesSize = -1;
-	bool IsRootPlugin = false;
+	bool bIsRootPlugin = false;
 };
 
 class FAssetTableRow
