Disk Size Tool -- Converting to use the new required install size tag for asset sizes where possible.

diff --git a/Engine/Plugins/Editor/AssetManagerEditor/Source/AssetManagerEditor/Private/TreeView/AssetTable.cpp b/Engine/Plugins/Editor/AssetManagerEditor/Source/AssetManagerEditor/Private/TreeView/AssetTable.cpp
--- a/Engine/Plugins/Editor/AssetManagerEditor/Source/AssetManagerEditor/Private/TreeView/AssetTable.cpp
+++ b/Engine/Plugins/Editor/AssetManagerEditor/Source/AssetManagerEditor/Private/TreeView/AssetTable.cpp
@@ -27,7 +27,7 @@ const FName FAssetTableColumns::NameColumnId(TEXT("Name"));
 const FName FAssetTableColumns::PathColumnId(TEXT("Path"));
 const FName FAssetTableColumns::PrimaryTypeColumnId(TEXT("PrimaryType"));
 const FName FAssetTableColumns::PrimaryNameColumnId(TEXT("PrimaryName"));
-const FName FAssetTableColumns::StagedCompressedSizeColumnId(TEXT("StagedCompressedSize"));
+const FName FAssetTableColumns::StagedCompressedSizeRequiredInstallColumnId(TEXT("StagedCompressedSizeRequiredInstall"));
 const FName FAssetTableColumns::TotalSizeUniqueDependenciesColumnId(TEXT("TotalSizeUniqueDependencies"));
 const FName FAssetTableColumns::TotalSizeSharedDependenciesColumnId(TEXT("TotalSizeSharedDependencies"));
 const FName FAssetTableColumns::TotalSizeExternalDependenciesColumnId(TEXT("TotalSizeExternalDependencies"));
@@ -203,7 +203,7 @@ FAssetTable::FAssetTable()
 		Asset.PrimaryType = StoreStr(FString::Printf(TEXT("PT_%02d"), Id2 % 10));
 		Asset.PrimaryName = StoreStr(FString::Printf(TEXT("PN%d"), Id2));
 		Asset.TotalUsageCount = 10;
-		Asset.StagedCompressedSize = 1;
+		Asset.StagedCompressedSizeRequiredInstall = 1;
 		Asset.NativeClass = StoreStr(FString::Printf(TEXT("NativeClass%02d"), (Id * Id * Id) % 8));
 		Asset.PluginName = StoreStr(TEXT("MockGFP"));
 	}
@@ -626,14 +626,14 @@ void FAssetTable::AddDefaultColumns()
 	//////////////////////////////////////////////////
 	// Staged Compressed Size Column
 	{
-		TSharedRef<FTableColumn> ColumnRef = MakeShared<FTableColumn>(FAssetTableColumns::StagedCompressedSizeColumnId);
+		TSharedRef<FTableColumn> ColumnRef = MakeShared<FTableColumn>(FAssetTableColumns::StagedCompressedSizeRequiredInstallColumnId);
 		FTableColumn& Column = *ColumnRef;
 
 		Column.SetIndex(ColumnIndex++);
 
-		Column.SetShortName(LOCTEXT("StagedCompressedSizeColumnName", "Self Size"));
-		Column.SetTitleName(LOCTEXT("StagedCompressedSizeColumnTitle", "Self Size (Compressed)"));
-		Column.SetDescription(LOCTEXT("StagedCompressedSizeColumnDesc", "Compressed size of iostore chunks for this asset's package. Only visible after staging."));
+		Column.SetShortName(LOCTEXT("StagedCompressedSizeRequiredInstallColumnName", "Self Size"));
+		Column.SetTitleName(LOCTEXT("StagedCompressedSizeRequiredInstallColumnTitle", "Self Size (Compressed, Required Install)"));
+		Column.SetDescription(LOCTEXT("StagedCompressedSizeRequiredInstallColumnDesc", "Compressed size of required install iostore chunks for this asset's package. Only visible after staging."));
 
 		Column.SetFlags(ETableColumnFlags::ShouldBeVisible | ETableColumnFlags::CanBeHidden | ETableColumnFlags::CanBeFiltered | ETableColumnFlags::IsDynamic);
 
@@ -664,7 +664,7 @@ void FAssetTable::AddDefaultColumns()
 					}
 					const FAssetTreeNode& TreeNode = Node.As<FAssetTreeNode>();
 					const FAssetTableRow& Asset = TreeNode.GetAssetChecked();
-					return FTableCellValue(static_cast<int64>(Asset.GetStagedCompressedSize()));
+					return FTableCellValue(static_cast<int64>(Asset.GetStagedCompressedSizeRequiredInstall()));
 				}
 				else if (Node.IsGroup() && !Node.Is<FPluginDependenciesGroupNode>())
 				{
@@ -1505,7 +1505,7 @@ void FAssetTable::AddDefaultColumns()
 
 	for (int32 Index : UniqueDependencies)
 	{
-		Result.UniqueDependenciesSize += OwningTable.GetAssetChecked(Index).GetStagedCompressedSize();
+		Result.UniqueDependenciesSize += OwningTable.GetAssetChecked(Index).GetStagedCompressedSizeRequiredInstall();
 	}
 	if (OutUniqueDependencies != nullptr)
 	{
@@ -1532,7 +1532,7 @@ void FAssetTable::AddDefaultColumns()
 			continue;
 		}
 
-		int64 DependencySize = OwningTable.GetAssetChecked(CurrentIndex).GetStagedCompressedSize();
+		int64 DependencySize = OwningTable.GetAssetChecked(CurrentIndex).GetStagedCompressedSizeRequiredInstall();
 		Result.SharedDependenciesSize += DependencySize;
 		if (OutSharedDependencies != nullptr)
 		{
@@ -1577,7 +1577,7 @@ void FAssetTable::AddDefaultColumns()
 			{
 				OutExternalDependencies->Add(Index);
 			}
-			TotalSizeExternalDependencies += Row.StagedCompressedSize;
+			TotalSizeExternalDependencies += Row.StagedCompressedSizeRequiredInstall;
 			NumDependenciesIncluded++;
 		}
 	}
diff --git a/Engine/Plugins/Editor/AssetManagerEditor/Source/AssetManagerEditor/Private/TreeView/AssetTable.h b/Engine/Plugins/Editor/AssetManagerEditor/Source/AssetManagerEditor/Private/TreeView/AssetTable.h
--- a/Engine/Plugins/Editor/AssetManagerEditor/Source/AssetManagerEditor/Private/TreeView/AssetTable.h
+++ b/Engine/Plugins/Editor/AssetManagerEditor/Source/AssetManagerEditor/Private/TreeView/AssetTable.h
@@ -117,7 +117,7 @@ struct FAssetTableColumns
 	static const FName PathColumnId;
 	static const FName PrimaryTypeColumnId;
 	static const FName PrimaryNameColumnId;
-	static const FName StagedCompressedSizeColumnId;
+	static const FName StagedCompressedSizeRequiredInstallColumnId;
 	static const FName TotalSizeUniqueDependenciesColumnId;
 	static const FName TotalSizeSharedDependenciesColumnId;
 	static const FName TotalSizeExternalDependenciesColumnId;
@@ -191,7 +191,7 @@ public:
 	const TCHAR* GetPath() const { return Path; }
 	const TCHAR* GetPrimaryType() const { return PrimaryType; }
 	const TCHAR* GetPrimaryName() const { return PrimaryName; }
-	int64 GetStagedCompressedSize() const { return StagedCompressedSize; }
+	int64 GetStagedCompressedSizeRequiredInstall() const { return StagedCompressedSizeRequiredInstall; }
 	int64 GetOrComputeTotalSizeUniqueDependencies(const FAssetTable& OwningTable, int32 ThisIndex) const;
 	int64 GetOrComputeTotalSizeSharedDependencies(const FAssetTable& OwningTable, int32 ThisIndex) const; 
 	int64 GetOrComputeTotalSizeExternalDependencies(const FAssetTable& OwningTable, int32 ThisIndex) const; 
@@ -226,7 +226,7 @@ private:
 	const TCHAR* Path = nullptr;
 	const TCHAR* PrimaryType = nullptr;
 	const TCHAR* PrimaryName = nullptr;
-	int64 StagedCompressedSize = 0;
+	int64 StagedCompressedSizeRequiredInstall = 0;
 	mutable int64 TotalSizeUniqueDependencies = -1; // Lazily calculated
 	mutable int64 TotalSizeSharedDependencies = -1; // Lazily calculated
 	mutable int64 TotalSizeExternalDependencies = -1; // Lazily calculated
diff --git a/Engine/Plugins/Editor/AssetManagerEditor/Source/AssetManagerEditor/Private/TreeView/SAssetTableTreeView.cpp b/Engine/Plugins/Editor/AssetManagerEditor/Source/AssetManagerEditor/Private/TreeView/SAssetTableTreeView.cpp
--- a/Engine/Plugins/Editor/AssetManagerEditor/Source/AssetManagerEditor/Private/TreeView/SAssetTableTreeView.cpp
+++ b/Engine/Plugins/Editor/AssetManagerEditor/Source/AssetManagerEditor/Private/TreeView/SAssetTableTreeView.cpp
@@ -262,7 +262,7 @@ void SAssetTableTreeView::InitAvailableViewPresets()
 		}
 		virtual FName GetSortColumn() const override
 		{
-			return FAssetTableColumns::StagedCompressedSizeColumnId;
+			return FAssetTableColumns::StagedCompressedSizeRequiredInstallColumnId;
 		}
 		virtual EColumnSortMode::Type GetSortMode() const override
 		{
@@ -301,7 +301,7 @@ void SAssetTableTreeView::InitAvailableViewPresets()
 
 			InOutConfigSet.Add({ UE::Insights::FTable::GetHierarchyColumnId(),              true, 400.0f });
 			InOutConfigSet.Add({ FAssetTableColumns::CountColumnId,                         true, 100.0f });
-			InOutConfigSet.Add({ FAssetTableColumns::StagedCompressedSizeColumnId,          true, 100.0f });
+			InOutConfigSet.Add({ FAssetTableColumns::StagedCompressedSizeRequiredInstallColumnId,          true, 100.0f });
 			InOutConfigSet.Add({ FAssetTableColumns::PluginInclusiveSizeColumnId,	        true, 100.0f });
 			InOutConfigSet.Add({ FAssetTableColumns::TypeColumnId,                          true, 200.0f });
 			InOutConfigSet.Add({ FAssetTableColumns::NameColumnId,                         !true, 200.0f });
@@ -384,13 +384,13 @@ void SAssetTableTreeView::InitAvailableViewPresets()
 		{
 			InOutConfigSet.Add({ UE::Insights::FTable::GetHierarchyColumnId(),              true, 400.0f });
 			InOutConfigSet.Add({ FAssetTableColumns::CountColumnId,                         true, 100.0f });
-			InOutConfigSet.Add({ FAssetTableColumns::StagedCompressedSizeColumnId,          true, 100.0f });
+			InOutConfigSet.Add({ FAssetTableColumns::StagedCompressedSizeRequiredInstallColumnId,          true, 100.0f });
 			InOutConfigSet.Add({ FAssetTableColumns::TypeColumnId,                          true, 200.0f });
 			InOutConfigSet.Add({ FAssetTableColumns::NameColumnId,                         !true, 200.0f });
 			InOutConfigSet.Add({ FAssetTableColumns::PathColumnId,                         !true, 400.0f });
 			InOutConfigSet.Add({ FAssetTableColumns::PrimaryTypeColumnId,                   true, 200.0f });
 			InOutConfigSet.Add({ FAssetTableColumns::PrimaryNameColumnId,                  !true, 200.0f });
-			InOutConfigSet.Add({ FAssetTableColumns::StagedCompressedSizeColumnId,          true, 100.0f });
+			InOutConfigSet.Add({ FAssetTableColumns::StagedCompressedSizeRequiredInstallColumnId,          true, 100.0f });
 			InOutConfigSet.Add({ FAssetTableColumns::TotalSizeUniqueDependenciesColumnId,   true, 100.0f });
 			InOutConfigSet.Add({ FAssetTableColumns::TotalSizeSharedDependenciesColumnId,   true, 100.0f });
 			InOutConfigSet.Add({ FAssetTableColumns::TotalSizeExternalDependenciesColumnId,!true, 100.0f });
@@ -418,7 +418,7 @@ void SAssetTableTreeView::InitAvailableViewPresets()
 		}
 		virtual FName GetSortColumn() const override
 		{
-			return FAssetTableColumns::StagedCompressedSizeColumnId;
+			return FAssetTableColumns::StagedCompressedSizeRequiredInstallColumnId;
 		}
 		virtual EColumnSortMode::Type GetSortMode() const override
 		{
@@ -457,7 +457,7 @@ void SAssetTableTreeView::InitAvailableViewPresets()
 		{
 			InOutConfigSet.Add({ UE::Insights::FTable::GetHierarchyColumnId(),              true, 400.0f });
 			InOutConfigSet.Add({ FAssetTableColumns::CountColumnId,                         true, 100.0f });
-			InOutConfigSet.Add({ FAssetTableColumns::StagedCompressedSizeColumnId,          true, 100.0f });
+			InOutConfigSet.Add({ FAssetTableColumns::StagedCompressedSizeRequiredInstallColumnId,          true, 100.0f });
 			InOutConfigSet.Add({ FAssetTableColumns::TypeColumnId,                          true, 200.0f });
 			InOutConfigSet.Add({ FAssetTableColumns::NameColumnId,                         !true, 200.0f });
 			InOutConfigSet.Add({ FAssetTableColumns::PathColumnId,                         !true, 400.0f });
@@ -489,7 +489,7 @@ void SAssetTableTreeView::InitAvailableViewPresets()
 		}
 		virtual FName GetSortColumn() const override
 		{
-			return FAssetTableColumns::StagedCompressedSizeColumnId;
+			return FAssetTableColumns::StagedCompressedSizeRequiredInstallColumnId;
 		}
 		virtual EColumnSortMode::Type GetSortMode() const override
 		{
@@ -528,7 +528,7 @@ void SAssetTableTreeView::InitAvailableViewPresets()
 		{
 			InOutConfigSet.Add({ UE::Insights::FTable::GetHierarchyColumnId(),              true, 400.0f });
 			InOutConfigSet.Add({ FAssetTableColumns::CountColumnId,                         true, 100.0f });
-			InOutConfigSet.Add({ FAssetTableColumns::StagedCompressedSizeColumnId,          true, 100.0f });
+			InOutConfigSet.Add({ FAssetTableColumns::StagedCompressedSizeRequiredInstallColumnId,          true, 100.0f });
 			InOutConfigSet.Add({ FAssetTableColumns::TypeColumnId,                          true, 200.0f });
 			InOutConfigSet.Add({ FAssetTableColumns::NameColumnId,                         !true, 200.0f });
 			InOutConfigSet.Add({ FAssetTableColumns::PathColumnId,                         !true, 400.0f });
@@ -561,7 +561,7 @@ void SAssetTableTreeView::InitAvailableViewPresets()
 		}
 		virtual FName GetSortColumn() const override
 		{
-			return FAssetTableColumns::StagedCompressedSizeColumnId;
+			return FAssetTableColumns::StagedCompressedSizeRequiredInstallColumnId;
 		}
 		virtual EColumnSortMode::Type GetSortMode() const override
 		{
@@ -600,7 +600,7 @@ void SAssetTableTreeView::InitAvailableViewPresets()
 		{
 			InOutConfigSet.Add({ UE::Insights::FTable::GetHierarchyColumnId(),              true, 400.0f });
 			InOutConfigSet.Add({ FAssetTableColumns::CountColumnId,                         true, 100.0f });
-			InOutConfigSet.Add({ FAssetTableColumns::StagedCompressedSizeColumnId,          true, 100.0f });
+			InOutConfigSet.Add({ FAssetTableColumns::StagedCompressedSizeRequiredInstallColumnId,          true, 100.0f });
 			InOutConfigSet.Add({ FAssetTableColumns::TypeColumnId,                          true, 200.0f });
 			InOutConfigSet.Add({ FAssetTableColumns::NameColumnId,                         !true, 200.0f });
 			InOutConfigSet.Add({ FAssetTableColumns::PathColumnId,                         !true, 400.0f });
@@ -634,7 +634,7 @@ void SAssetTableTreeView::InitAvailableViewPresets()
 		}
 		virtual FName GetSortColumn() const override
 		{
-			return FAssetTableColumns::StagedCompressedSizeColumnId;
+			return FAssetTableColumns::StagedCompressedSizeRequiredInstallColumnId;
 		}
 		virtual EColumnSortMode::Type GetSortMode() const override
 		{
@@ -673,7 +673,7 @@ void SAssetTableTreeView::InitAvailableViewPresets()
 		{
 			InOutConfigSet.Add({ UE::Insights::FTable::GetHierarchyColumnId(),              true, 400.0f });
 			InOutConfigSet.Add({ FAssetTableColumns::CountColumnId,                         true, 100.0f });
-			InOutConfigSet.Add({ FAssetTableColumns::StagedCompressedSizeColumnId,          true, 100.0f });
+			InOutConfigSet.Add({ FAssetTableColumns::StagedCompressedSizeRequiredInstallColumnId,          true, 100.0f });
 			InOutConfigSet.Add({ FAssetTableColumns::PluginInclusiveSizeColumnId,	        true, 100.0f });
 			InOutConfigSet.Add({ FAssetTableColumns::TypeColumnId,                          true, 200.0f });
 			InOutConfigSet.Add({ FAssetTableColumns::NameColumnId,                          true, 200.0f });
@@ -796,7 +796,7 @@ static void WriteDependencyLine(const FAssetTable& AssetTable, const TMap<int32,
 	ReusableLineBuffer->Appendf("%s%s,%s,%s,%s", *WriteToAnsiString<512>(Row.GetPath()),
 												*WriteToAnsiString<64>(Row.GetName()), 
 												*WriteToAnsiString<64>(Row.GetType()),
-												*WriteToAnsiString<32>(LexToString(Row.GetStagedCompressedSize())),
+												*WriteToAnsiString<32>(LexToString(Row.GetStagedCompressedSizeRequiredInstall())),
 												*WriteToAnsiString<16>(*DependencyType));
 
 	if (const TArray<int32>* Route = RouteMap.Find(RowIndex))
@@ -1479,11 +1479,11 @@ void SAssetTableTreeView::CalculateBaseAndMarginalCostForSelection(TSet<int32>&
 	{
 		if (ReferenceCountPair.Value > 1)
 		{
-			*OutTotalSizeMultiplyUsed += GetAssetTable()->GetAssetChecked(ReferenceCountPair.Key).GetStagedCompressedSize();
+			*OutTotalSizeMultiplyUsed += GetAssetTable()->GetAssetChecked(ReferenceCountPair.Key).GetStagedCompressedSizeRequiredInstall();
 		}
 		else if (ReferenceCountPair.Value == 1)
 		{
-			*OutTotalSizeSingleUse += GetAssetTable()->GetAssetChecked(ReferenceCountPair.Key).GetStagedCompressedSize();
+			*OutTotalSizeSingleUse += GetAssetTable()->GetAssetChecked(ReferenceCountPair.Key).GetStagedCompressedSizeRequiredInstall();
 		}
 	}
 }
@@ -1570,7 +1570,7 @@ void SAssetTableTreeView::TreeView_OnSelectionChanged(UE::Insights::FTableTreeNo
 		FooterCenterText1 = FText::FromString(AssetTableRow.GetPath());
 		FooterCenterText2 = FText::FromString(AssetTableRow.GetName());
 		FooterRightText1 = FText::Format(LOCTEXT("FooterRightFmt", "Self: {0}    Unique: {1}    Shared: {2}    External: {3}"),
-			FText::AsMemory(AssetTableRow.GetStagedCompressedSize()),
+			FText::AsMemory(AssetTableRow.GetStagedCompressedSizeRequiredInstall()),
 			FText::AsMemory(AssetTableRow.GetOrComputeTotalSizeUniqueDependencies(*GetAssetTable(), NewlySelectedAssetRowIndex)),
 			FText::AsMemory(AssetTableRow.GetOrComputeTotalSizeSharedDependencies(*GetAssetTable(), NewlySelectedAssetRowIndex)),
 			FText::AsMemory(AssetTableRow.GetOrComputeTotalSizeExternalDependencies(*GetAssetTable(), NewlySelectedAssetRowIndex)));
@@ -1607,7 +1607,7 @@ void SAssetTableTreeView::TreeView_OnSelectionChanged(UE::Insights::FTableTreeNo
 		int64 TotalSelfSize = 0;
 		for (int32 Index : SelectedAssetIndices)
 		{
-			TotalSelfSize += GetAssetTable()->GetAssetChecked(Index).GetStagedCompressedSize();
+			TotalSelfSize += GetAssetTable()->GetAssetChecked(Index).GetStagedCompressedSizeRequiredInstall();
 		}
 
 		FText BaseAndMarginalCost;
@@ -1707,7 +1707,12 @@ void SAssetTableTreeView::PopulateAssetTableRow(FAssetTableRow& OutRow, const FA
 		OutRow.PrimaryName = AssetTable.StoreStr(Str);
 	}
 
-	EditorModule->GetIntegerValueForCustomColumn(AssetData, UE::AssetRegistry::Stage_ChunkCompressedSizeFName, OutRow.StagedCompressedSize, &RegistrySource);
+	bool RegistryHasSeparateCompressedSizes = EditorModule->GetIntegerValueForCustomColumn(AssetData, UE::AssetRegistry::Stage_ChunkInstalledSizeFName, OutRow.StagedCompressedSizeRequiredInstall, &RegistrySource);
+	if (!RegistryHasSeparateCompressedSizes)
+	{
+		// This should only apply to legacy asset registries (before ~6/28/23).
+		EditorModule->GetIntegerValueForCustomColumn(AssetData, UE::AssetRegistry::Stage_ChunkCompressedSizeFName, OutRow.StagedCompressedSizeRequiredInstall, &RegistrySource);
+	}
 	EditorModule->GetIntegerValueForCustomColumn(AssetData, IAssetManagerEditorModule::TotalUsageName, OutRow.TotalUsageCount, &RegistrySource);
 
 	if (EditorModule->GetStringValueForCustomColumn(AssetData, IAssetManagerEditorModule::ChunksName, Str, &RegistrySource))
@@ -1774,11 +1779,11 @@ void SAssetTableTreeView::RefreshAssets()
 
 					if (int64* PluginSize = PluginToSizeMap.Find(AssetRow.GetPluginName()))
 					{
-						(*PluginSize) += AssetRow.GetStagedCompressedSize();
+						(*PluginSize) += AssetRow.GetStagedCompressedSizeRequiredInstall();
 					}
 					else
 					{
-						PluginToSizeMap.Add(AssetRow.GetPluginName(), AssetRow.GetStagedCompressedSize());
+						PluginToSizeMap.Add(AssetRow.GetPluginName(), AssetRow.GetStagedCompressedSizeRequiredInstall());
 					}
 				}
 			}
@@ -1834,11 +1839,11 @@ void SAssetTableTreeView::RefreshAssets()
 						SourceAssets.Add(*DependencyAsset);
 						if (int64* PluginSize = PluginToSizeMap.Find(NewRow.GetPluginName()))
 						{
-							(*PluginSize) += NewRow.GetStagedCompressedSize();
+							(*PluginSize) += NewRow.GetStagedCompressedSizeRequiredInstall();
 						}
 						else
 						{
-							PluginToSizeMap.Add(NewRow.GetPluginName(), NewRow.GetStagedCompressedSize());
+							PluginToSizeMap.Add(NewRow.GetPluginName(), NewRow.GetStagedCompressedSizeRequiredInstall());
 						}
 					}
 				}
@@ -1902,6 +1907,8 @@ void SAssetTableTreeView::RefreshAssets()
 	if (CookMetadata.IsValid())
 	{
 		const UE::Cook::FCookMetadataPluginHierarchy& PluginHierarchy = CookMetadata.GetPluginHierarchy();
+		int64 TotalMismatch = 0;
+		int64 TotalSize = 0;
 		for (const UE::Cook::FCookMetadataPluginEntry& PluginEntry : PluginHierarchy.PluginsEnabledAtCook)
 		{
 			const TCHAR* StoredPluginName = AssetTable->StoreStr(PluginEntry.Name);
@@ -1922,15 +1929,14 @@ void SAssetTableTreeView::RefreshAssets()
 			{
 				if (CookMetadata.GetSizesPresent() != UE::Cook::ECookMetadataSizesPresent::NotPresent)
 				{
-					int64 TotalSizeOfPluginInMetadata = PluginEntry.ExclusiveSizes[UE::Cook::EPluginSizeTypes::Streaming]
-														+ PluginEntry.ExclusiveSizes[UE::Cook::EPluginSizeTypes::Installed]
-														+ PluginEntry.ExclusiveSizes[UE::Cook::EPluginSizeTypes::Optional]
-														+ PluginEntry.ExclusiveSizes[UE::Cook::EPluginSizeTypes::OptionalSegment];
-					if (FMath::Abs((*SizePtr - TotalSizeOfPluginInMetadata)) > static_cast<int64>(0.05 * FMath::Max(*SizePtr, TotalSizeOfPluginInMetadata)))
+					int64 TotalSizeOfPluginInMetadata = PluginEntry.ExclusiveSizes[UE::Cook::EPluginSizeTypes::Installed];
+					if (FMath::Abs((*SizePtr - TotalSizeOfPluginInMetadata)) > static_cast<int64>(0.01 * FMath::Max(*SizePtr, TotalSizeOfPluginInMetadata)))
 					{
 						UE_LOG(LogInsights, Warning, TEXT("Plugin %s found with ucookmetadata and asset calculation size delta > 5%%. Metadata size: %lld Calculated size: %lld"), 
 							StoredPluginName, PluginInfo.Size, *SizePtr);
 					}
+					TotalMismatch += FMath::Abs(*SizePtr - TotalSizeOfPluginInMetadata);
+					TotalSize += TotalSizeOfPluginInMetadata;
 				}
 			}
 
@@ -1952,6 +1958,8 @@ void SAssetTableTreeView::RefreshAssets()
 			}
 		}
 
+		UE_LOG(LogInsights, Warning, TEXT("Total Size of Plugins from Metadata: %lld // Total Discrepancy (vs calculated size from assets): %lld"), TotalSize, TotalMismatch);
+
 		for (uint16 RootPluginIndex : PluginHierarchy.RootPlugins)
 		{
 			const UE::Cook::FCookMetadataPluginEntry& PluginEntry = PluginHierarchy.PluginsEnabledAtCook[RootPluginIndex];
