Asset Disk Size Tool -- Fix a bug where we were sometimes using the in memory asset registry causing us to miss some assets in the loaded DevAR data.

diff --git a/Engine/Plugins/Editor/AssetManagerEditor/Source/AssetManagerEditor/Private/TreeView/SAssetTableTreeView.cpp b/Engine/Plugins/Editor/AssetManagerEditor/Source/AssetManagerEditor/Private/TreeView/SAssetTableTreeView.cpp
--- a/Engine/Plugins/Editor/AssetManagerEditor/Source/AssetManagerEditor/Private/TreeView/SAssetTableTreeView.cpp
+++ b/Engine/Plugins/Editor/AssetManagerEditor/Source/AssetManagerEditor/Private/TreeView/SAssetTableTreeView.cpp
@@ -66,7 +66,6 @@ void SAssetTableTreeView::Construct(const FArguments& InArgs, TSharedPtr<FAssetT
 	this->OnSelectionChanged = InArgs._OnSelectionChanged;
 
 	FAssetRegistryModule& AssetRegistryModule = FModuleManager::LoadModuleChecked<FAssetRegistryModule>(TEXT("AssetRegistry"));
-	AssetRegistry = &AssetRegistryModule.Get();
 	AssetManager = &UAssetManager::Get();
 	EditorModule = &IAssetManagerEditorModule::Get();
 
@@ -1765,17 +1764,16 @@ void SAssetTableTreeView::RefreshAssets()
 
 		for (int32 SourceAssetIndex = 0; SourceAssetIndex < SourceAssets.Num(); SourceAssetIndex++)
 		{
-			TArray<FAssetData> AssetsInSourcePackage;
-			AssetRegistry->GetAssetsByPackageName(SourceAssets[SourceAssetIndex].PackageName, AssetsInSourcePackage, /*bIncludeOnlyOnDiskAssets*/ true); // Only use on disk assets to avoid creating FAssetData for everything in memory
+			TArrayView<FAssetData const* const> AssetsInSourcePackage = RegistrySource.GetOwnedRegistryState()->GetAssetsByPackageName(SourceAssets[SourceAssetIndex].PackageName);
 
-			for (FAssetData& SourceAsset : AssetsInSourcePackage)
+			for (FAssetData const* const SourceAsset : AssetsInSourcePackage)
 			{
-				if (AssetToIndexMap.Find(SourceAsset) == nullptr)
+				if (AssetToIndexMap.Find(*SourceAsset) == nullptr)
 				{
 					FAssetTableRow AssetRow;
-					PopulateAssetTableRow(AssetRow, SourceAsset, *AssetTable);
+					PopulateAssetTableRow(AssetRow, *SourceAsset, *AssetTable);
 					AssetTable->AddAsset(AssetRow);
-					AssetToIndexMap.Add(SourceAsset, AssetTable->GetTotalAssetCount() - 1);
+					AssetToIndexMap.Add(*SourceAsset, AssetTable->GetTotalAssetCount() - 1);
 
 					if (int64* PluginSize = PluginToSizeMap.Find(AssetRow.GetPluginName()))
 					{
@@ -1809,8 +1807,7 @@ void SAssetTableTreeView::RefreshAssets()
 
 			FAssetIdentifier CurrentIdentifier(SourceAssets[SourceAssetIndex].PackageName);
 			// We'll have to add the dependencies to all these entries
-			TArray<FAssetData> AssetsInSourcePackage;
-			AssetRegistry->GetAssetsByPackageName(SourceAssets[SourceAssetIndex].PackageName, AssetsInSourcePackage, /*bIncludeOnlyOnDiskAssets*/ true); // Only use on disk assets to avoid creating FAssetData for everything in memory
+			TArrayView<FAssetData const* const> AssetsInSourcePackage = RegistrySource.GetOwnedRegistryState()->GetAssetsByPackageName(SourceAssets[SourceAssetIndex].PackageName);
 
 			// Get the dependencies
 			TArray<FAssetIdentifier> DependencyList;
@@ -1852,12 +1849,12 @@ void SAssetTableTreeView::RefreshAssets()
 			if (IndicesOfDependenciesInRowTableToAddToCurrentSourceAssetRow.Num())
 			{
 				// We found some dependencies. Let's add them.
-				for (const FAssetData& SourceAssetData : AssetsInSourcePackage)
+				for (const FAssetData* const SourceAssetData : AssetsInSourcePackage)
 				{
-					int32* RowIndex = AssetToIndexMap.Find(SourceAssetData);
+					int32* RowIndex = AssetToIndexMap.Find(*SourceAssetData);
 					if (RowIndex == nullptr)
 					{
-						UE_LOG(LogInsights, Warning, TEXT("Failed to find asset %s in package %s, source asset index %d. Asset registry loading was %s"), *SourceAssetData.AssetName.ToString(), *SourceAssetData.PackageName.ToString(), SourceAssetIndex, AssetRegistry->IsLoadingAssets() ? TEXT("INCOMPLETE") : TEXT("complete"));
+						UE_LOG(LogInsights, Warning, TEXT("Failed to find asset %s in package %s, source asset index %d."), *SourceAssetData->AssetName.ToString(), *SourceAssetData->PackageName.ToString(), SourceAssetIndex);
 					}
 
 					if (ensure(RowIndex != nullptr))
@@ -1889,9 +1886,9 @@ void SAssetTableTreeView::RefreshAssets()
 					// Now for each of those dependencies, add this source asset row as a referencer
 					if (FAssetTableRow* DependentRow = AssetTable->GetAsset(Index))
 					{
-						for (const FAssetData& SourceAssetData : AssetsInSourcePackage)
+						for (const FAssetData* SourceAssetData : AssetsInSourcePackage)
 						{
-							int32* RowIndex = AssetToIndexMap.Find(SourceAssetData);
+							int32* RowIndex = AssetToIndexMap.Find(*SourceAssetData);
 							if (ensure(RowIndex != nullptr))
 							{
 								DependentRow->Referencers.AddUnique(*RowIndex);
@@ -1923,16 +1920,16 @@ void SAssetTableTreeView::RefreshAssets()
 			}
 
 			///
-			// Temp validation
+			/// validation
 			///
 			if (const int64* SizePtr = PluginToSizeMap.Find(PluginEntry.Name))
 			{
 				if (CookMetadata.GetSizesPresent() != UE::Cook::ECookMetadataSizesPresent::NotPresent)
 				{
 					int64 TotalSizeOfPluginInMetadata = PluginEntry.ExclusiveSizes[UE::Cook::EPluginSizeTypes::Installed];
-					if (FMath::Abs((*SizePtr - TotalSizeOfPluginInMetadata)) > static_cast<int64>(0.01 * FMath::Max(*SizePtr, TotalSizeOfPluginInMetadata)))
+					if (*SizePtr != TotalSizeOfPluginInMetadata)
 					{
-						UE_LOG(LogInsights, Warning, TEXT("Plugin %s found with ucookmetadata and asset calculation size delta > 5%%. Metadata size: %lld Calculated size: %lld"), 
+						UE_LOG(LogInsights, Warning, TEXT("Plugin %s found with mismatched ucookmetadata and internal asset size calculation. Metadata size: %lld Calculated size: %lld"), 
 							StoredPluginName, PluginInfo.Size, *SizePtr);
 					}
 					TotalMismatch += FMath::Abs(*SizePtr - TotalSizeOfPluginInMetadata);
@@ -1958,14 +1955,17 @@ void SAssetTableTreeView::RefreshAssets()
 			}
 		}
 
-		UE_LOG(LogInsights, Warning, TEXT("Total Size of Plugins from Metadata: %lld // Total Discrepancy (vs calculated size from assets): %lld"), TotalSize, TotalMismatch);
+		if (TotalMismatch != 0)
+		{
+			UE_LOG(LogInsights, Warning, TEXT("Total Size of Plugins from Metadata: %lld // Total Discrepancy (vs calculated size from assets): %lld"), TotalSize, TotalMismatch);
+		}
 
 		for (uint16 RootPluginIndex : PluginHierarchy.RootPlugins)
 		{
 			const UE::Cook::FCookMetadataPluginEntry& PluginEntry = PluginHierarchy.PluginsEnabledAtCook[RootPluginIndex];
 			const TCHAR* StoredPluginName = AssetTable->StoreStr(PluginEntry.Name);
 			FAssetTablePluginInfo& PluginInfo = AssetTable->GetOrCreatePluginInfo(StoredPluginName);
-			UE_LOG(LogInsights, Warning, TEXT("Found root plugin %s"), StoredPluginName);
+			UE_LOG(LogInsights, Display, TEXT("Found root plugin %s"), StoredPluginName);
 			PluginInfo.bIsRootPlugin = true;
 		}
 	}
diff --git a/Engine/Plugins/Editor/AssetManagerEditor/Source/AssetManagerEditor/Private/TreeView/SAssetTableTreeView.h b/Engine/Plugins/Editor/AssetManagerEditor/Source/AssetManagerEditor/Private/TreeView/SAssetTableTreeView.h
--- a/Engine/Plugins/Editor/AssetManagerEditor/Source/AssetManagerEditor/Private/TreeView/SAssetTableTreeView.h
+++ b/Engine/Plugins/Editor/AssetManagerEditor/Source/AssetManagerEditor/Private/TreeView/SAssetTableTreeView.h
@@ -165,7 +165,6 @@ private:
 	UE::Cook::FCookMetadataState CookMetadata;
 
 	/** Cached interfaces */
-	class IAssetRegistry* AssetRegistry;
 	class UAssetManager* AssetManager;
 	class IAssetManagerEditorModule* EditorModule;
 };
