Disk Size Tool -- Renaming "By Dependency" grouping to "By [REDACTED_PERSON]".

diff --git a/Engine/Plugins/Editor/AssetManagerEditor/Source/AssetManagerEditor/Private/TreeView/AssetDependencyGrouping.cpp b/Engine/Plugins/Editor/AssetManagerEditor/Source/AssetManagerEditor/Private/TreeView/AssetDependencyGrouping.cpp
--- a/Engine/Plugins/Editor/AssetManagerEditor/Source/AssetManagerEditor/Private/TreeView/AssetDependencyGrouping.cpp
+++ b/Engine/Plugins/Editor/AssetManagerEditor/Source/AssetManagerEditor/Private/TreeView/AssetDependencyGrouping.cpp
@@ -17,9 +17,9 @@ INSIGHTS_IMPLEMENT_RTTI(FPluginDependencyGrouping)
 
 FAssetDependencyGrouping::FAssetDependencyGrouping()
 	: FTreeNodeGrouping(
-		LOCTEXT("Grouping_ByDependency_ShortName", "Dependency"),
-		LOCTEXT("Grouping_ByDependency_TitleName", "By Dependency"),
-		LOCTEXT("Grouping_ByDependency_Desc", "Group assets based on their dependency."),
+		LOCTEXT("Grouping_ByAssetDependency_ShortName", "Asset Dependency"),
+		LOCTEXT("Grouping_ByAssetDependency_TitleName", "By Asset Dependency"),
+		LOCTEXT("Grouping_ByAssetDependency_Desc", "Group assets based on their dependency."),
 		TEXT("Icons.Group.TreeItem"),
 		nullptr)
 {
