Fix a bug in the Reference Viewer plugin filter

The list of available plugins to filter by was never populated.

diff --git a/Engine/Plugins/Editor/AssetManagerEditor/Source/AssetManagerEditor/Private/ReferenceViewer/EdGraph_ReferenceViewer.cpp b/Engine/Plugins/Editor/AssetManagerEditor/Source/AssetManagerEditor/Private/ReferenceViewer/EdGraph_ReferenceViewer.cpp
--- a/Engine/Plugins/Editor/AssetManagerEditor/Source/AssetManagerEditor/Private/ReferenceViewer/EdGraph_ReferenceViewer.cpp
+++ b/Engine/Plugins/Editor/AssetManagerEditor/Source/AssetManagerEditor/Private/ReferenceViewer/EdGraph_ReferenceViewer.cpp
@@ -214,7 +214,6 @@ UEdGraphNode_Reference* UEdGraph_ReferenceViewer::ConstructNodes(const TArray<FA
 		}
 
 		// Prepare for plugin filtering.
-		if (ShouldFilterByPlugin())
 		{
 			// Collect plugin names from assets reachable in the graph if the graph had been unfiltered.
 			EncounteredPluginsAmongNodes.Empty();
diff --git a/Engine/Plugins/Editor/AssetManagerEditor/Source/AssetManagerEditor/Private/ReferenceViewer/SReferenceViewer.cpp b/Engine/Plugins/Editor/AssetManagerEditor/Source/AssetManagerEditor/Private/ReferenceViewer/SReferenceViewer.cpp
--- a/Engine/Plugins/Editor/AssetManagerEditor/Source/AssetManagerEditor/Private/ReferenceViewer/SReferenceViewer.cpp
+++ b/Engine/Plugins/Editor/AssetManagerEditor/Source/AssetManagerEditor/Private/ReferenceViewer/SReferenceViewer.cpp
@@ -493,7 +493,7 @@ void SReferenceViewer::Construct(const FArguments& InArgs)
 						.VAlign(VAlign_Center)
 						.Padding(2.f)
 						[
-							SAssignNew(PluginFilterCheckbox, SCheckBox)
+							SNew(SCheckBox)
 							.OnCheckStateChanged( this, &SReferenceViewer::OnEnablePluginFilterChanged )
 							.IsChecked( this, &SReferenceViewer::IsEnablePluginFilterChecked )
 						]
diff --git a/Engine/Plugins/Editor/AssetManagerEditor/Source/AssetManagerEditor/Private/ReferenceViewer/SReferenceViewer.h b/Engine/Plugins/Editor/AssetManagerEditor/Source/AssetManagerEditor/Private/ReferenceViewer/SReferenceViewer.h
--- a/Engine/Plugins/Editor/AssetManagerEditor/Source/AssetManagerEditor/Private/ReferenceViewer/SReferenceViewer.h
+++ b/Engine/Plugins/Editor/AssetManagerEditor/Source/AssetManagerEditor/Private/ReferenceViewer/SReferenceViewer.h
@@ -8,7 +8,6 @@
 #include "CollectionManagerTypes.h"
 
 class FUICommandList;
-class SCheckBox;
 class SComboButton;
 class SReferenceViewerFilterBar;
 class SSearchBox;
@@ -240,7 +239,6 @@ private:
 	TArray<TSharedPtr<FName>> CollectionsComboList;
 
 	/** Combo box for plugin filter options */
-	TSharedPtr<SCheckBox> PluginFilterCheckbox;
 	TSharedPtr<SComboButton> PluginsComboButton;
 
 	/**
