#UE Do not follow loops when determining plugins in reference viewer graph

diff --git a/Engine/Plugins/Editor/AssetManagerEditor/Source/AssetManagerEditor/Private/ReferenceViewer/EdGraph_ReferenceViewer.cpp b/Engine/Plugins/Editor/AssetManagerEditor/Source/AssetManagerEditor/Private/ReferenceViewer/EdGraph_ReferenceViewer.cpp
--- a/Engine/Plugins/Editor/AssetManagerEditor/Source/AssetManagerEditor/Private/ReferenceViewer/EdGraph_ReferenceViewer.cpp
+++ b/Engine/Plugins/Editor/AssetManagerEditor/Source/AssetManagerEditor/Private/ReferenceViewer/EdGraph_ReferenceViewer.cpp
@@ -653,12 +653,17 @@ bool UEdGraph_ReferenceViewer::IsAssetPassingSearchTextFilter(const FAssetIdenti
 	return true;
 }
 
-void UEdGraph_ReferenceViewer::GetUnfilteredGraphPluginNamesRecursive(bool bReferencers, const FAssetIdentifier& InAssetIdentifier, int32 InCurrentDepth, int32 InMaxDepth, const FAssetManagerDependencyQuery& Query, TArray<FAssetIdentifier>& OutAssetIdentifiers)
+void UEdGraph_ReferenceViewer::GetUnfilteredGraphPluginNamesRecursive(bool bReferencers, const FAssetIdentifier& InAssetIdentifier, int32 InCurrentDepth, int32 InMaxDepth, const FAssetManagerDependencyQuery& Query, TSet<FAssetIdentifier>& OutAssetIdentifiers)
 {
 	if (ExceedsMaxSearchDepth(InCurrentDepth, InMaxDepth))
 	{
 		return;
 	}
+
+	if (OutAssetIdentifiers.Contains(InAssetIdentifier))
+	{
+		return;
+	}
 	
 	IAssetRegistry& AssetRegistry = FModuleManager::LoadModuleChecked<FAssetRegistryModule>("AssetRegistry").Get();
 	TArray<FAssetDependency> LinksToAsset;
@@ -691,11 +696,16 @@ void UEdGraph_ReferenceViewer::GetUnfilteredGraphPluginNames(TArray<FAssetIdenti
 	
 	const FAssetManagerDependencyQuery Query = GetReferenceSearchFlags(false);
 	
-	TArray<FAssetIdentifier> AssetIdentifiers;
+	TSet<FAssetIdentifier> AssetIdentifiers;
 	for (const FAssetIdentifier& RootIdentifier : RootIdentifiers)
 	{
-		GetUnfilteredGraphPluginNamesRecursive(true, RootIdentifier, 0, Settings->GetSearchReferencerDepthLimit(), Query, AssetIdentifiers);
-		GetUnfilteredGraphPluginNamesRecursive(false, RootIdentifier, 0, Settings->GetSearchDependencyDepthLimit(), Query, AssetIdentifiers);
+		TSet<FAssetIdentifier> AssetReferencerIdentifiers;
+		GetUnfilteredGraphPluginNamesRecursive(true, RootIdentifier, 0, Settings->GetSearchReferencerDepthLimit(), Query, AssetReferencerIdentifiers);
+		AssetIdentifiers.Append(AssetReferencerIdentifiers);
+
+		TSet<FAssetIdentifier> AssetDependencyIdentifiers;
+		GetUnfilteredGraphPluginNamesRecursive(false, RootIdentifier, 0, Settings->GetSearchDependencyDepthLimit(), Query, AssetDependencyIdentifiers);
+		AssetIdentifiers.Append(AssetDependencyIdentifiers);
 	}
 
 	for (const FAssetIdentifier& AssetIdentifier : AssetIdentifiers)
diff --git a/Engine/Plugins/Editor/AssetManagerEditor/Source/AssetManagerEditor/Public/ReferenceViewer/EdGraph_ReferenceViewer.h b/Engine/Plugins/Editor/AssetManagerEditor/Source/AssetManagerEditor/Public/ReferenceViewer/EdGraph_ReferenceViewer.h
--- a/Engine/Plugins/Editor/AssetManagerEditor/Source/AssetManagerEditor/Public/ReferenceViewer/EdGraph_ReferenceViewer.h
+++ b/Engine/Plugins/Editor/AssetManagerEditor/Source/AssetManagerEditor/Public/ReferenceViewer/EdGraph_ReferenceViewer.h
@@ -155,7 +155,7 @@ private:
 	/** Returns true if filtering is enabled and we have a valid plugin name filter set */
 	bool ShouldFilterByPlugin() const;
 	
-	void GetUnfilteredGraphPluginNamesRecursive(bool bReferencers, const FAssetIdentifier& InAssetIdentifier, int32 InCurrentDepth, int32 InMaxDepth, const FAssetManagerDependencyQuery& Query, TArray<FAssetIdentifier>& OutAssetIdentifiers);
+	void GetUnfilteredGraphPluginNamesRecursive(bool bReferencers, const FAssetIdentifier& InAssetIdentifier, int32 InCurrentDepth, int32 InMaxDepth, const FAssetManagerDependencyQuery& Query, TSet<FAssetIdentifier>& OutAssetIdentifiers);
 	void GetUnfilteredGraphPluginNames(TArray<FAssetIdentifier> RootIdentifiers, TArray<FName>& OutPluginNames);
 
 	void GetSortedLinks(const TArray<FAssetIdentifier>& Identifiers, bool bReferencers, const FAssetManagerDependencyQuery& Query, TMap<FAssetIdentifier, EDependencyPinCategory>& OutLinks) const;
