AssetManagerEditor: Fixed grouping for lazy created child nodes in "Plugin Dependency Analysis". The parent node (which lazy creates child nodes) requires to implement GetAuthorGrouping().

diff --git a/Engine/Plugins/Editor/AssetManagerEditor/Source/AssetManagerEditor/Private/Insights/Filter/Widgets/SFilterConfiguratorRow.cpp b/Engine/Plugins/Editor/AssetManagerEditor/Source/AssetManagerEditor/Private/Insights/Filter/Widgets/SFilterConfiguratorRow.cpp
--- a/Engine/Plugins/Editor/AssetManagerEditor/Source/AssetManagerEditor/Private/Insights/Filter/Widgets/SFilterConfiguratorRow.cpp
+++ b/Engine/Plugins/Editor/AssetManagerEditor/Source/AssetManagerEditor/Private/Insights/Filter/Widgets/SFilterConfiguratorRow.cpp
@@ -227,8 +227,8 @@ TSharedRef<SWidget> SFilterConfiguratorRow::GenerateWidgetForColumn(const FName&
 					.VAlign(VAlign_Center)
 					[
 						SNew(SImage)
-						//TODO: .Image(FInsightsStyle::GetBrush("Icons.FilterAddGroup"))
-						.Image(FAppStyle::Get().GetBrush("Icons.PlusCircle"))
+						.Image(FInsightsStyle::GetBrush("Icons.FilterAddGroup"))
+						//.Image(FAppStyle::Get().GetBrush("Icons.PlusCircle"))
 					]
 
 					+ SHorizontalBox::Slot()
diff --git a/Engine/Plugins/Editor/AssetManagerEditor/Source/AssetManagerEditor/Private/Insights/Table/ViewModels/TableTreeNode.h b/Engine/Plugins/Editor/AssetManagerEditor/Source/AssetManagerEditor/Private/Insights/Table/ViewModels/TableTreeNode.h
--- a/Engine/Plugins/Editor/AssetManagerEditor/Source/AssetManagerEditor/Private/Insights/Table/ViewModels/TableTreeNode.h
+++ b/Engine/Plugins/Editor/AssetManagerEditor/Source/AssetManagerEditor/Private/Insights/Table/ViewModels/TableTreeNode.h
@@ -14,6 +14,7 @@ namespace Insights
 {
 
 class FTable;
+class FTreeNodeGrouping;
 class STableTreeView;
 
 struct FTableRowId
@@ -150,6 +151,13 @@ public:
 
 	virtual bool OnLazyCreateChildren(TSharedPtr<STableTreeView> InTableTreeView) { return false; }
 
+	/**
+	 * The grouping that has generated this node.
+	 * This is used to correctly apply further groupings for the lazy created children.
+	 * If this returns nullptr, grouping is not applied for lazy created children.
+	 */
+	virtual const FTreeNodeGrouping* GetAuthorGrouping() const { return nullptr; }
+
 protected:
 	TWeakPtr<FTable> ParentTable;
 	FTableRowId RowId;
diff --git a/Engine/Plugins/Editor/AssetManagerEditor/Source/AssetManagerEditor/Private/Insights/Table/Widgets/STableTreeView.cpp b/Engine/Plugins/Editor/AssetManagerEditor/Source/AssetManagerEditor/Private/Insights/Table/Widgets/STableTreeView.cpp
--- a/Engine/Plugins/Editor/AssetManagerEditor/Source/AssetManagerEditor/Private/Insights/Table/Widgets/STableTreeView.cpp
+++ b/Engine/Plugins/Editor/AssetManagerEditor/Source/AssetManagerEditor/Private/Insights/Table/Widgets/STableTreeView.cpp
@@ -142,16 +142,19 @@ STableTreeView::STableTreeView()
 	, ColumnBeingSorted(GetDefaultColumnBeingSorted())
 	, ColumnSortMode(GetDefaultColumnSortMode())
 {
-	AsyncOperationMutex = MakeShared<FRWLock>();
 }
 
 ////////////////////////////////////////////////////////////////////////////////////////////////////
 
 STableTreeView::~STableTreeView()
 {
-	CancelCurrentAsyncOp();
+	if (!bIsCloseScheduled)
+	{
+		UE_LOG(LogInsights, Log, TEXT("TableTreeView running in async mode was closed but OnClose() was not called. Call OnClose() from the owner tab/window."));
+	}
 
-	FScopedExclusiveRWLock AsyncOpLock(AsyncOperationMutex);
+	// Backup call to OnClose() in case it was not called from the owner.
+	OnClose();
 
 	if (CurrentAsyncOpFilterConfigurator)
 	{
@@ -1001,7 +1004,7 @@ FGraphEventRef STableTreeView::StartNodeFilteringTask(FGraphEventRef Prerequisit
 		Prerequisites.Add(DispatchEvent);
 	}
 
-	return TGraphTask<FTableTreeViewNodeFilteringAsyncTask>::CreateTask(&Prerequisites).ConstructAndDispatchWhenReady(this, AsyncOperationMutex);
+	return TGraphTask<FTableTreeViewNodeFilteringAsyncTask>::CreateTask(&Prerequisites).ConstructAndDispatchWhenReady(this);
 }
 
 ////////////////////////////////////////////////////////////////////////////////////////////////////
@@ -1266,7 +1269,7 @@ FGraphEventRef STableTreeView::StartHierarchyFilteringTask(FGraphEventRef Prereq
 		Prerequisites.Add(DispatchEvent);
 	}
 
-	return TGraphTask<FTableTreeViewHierarchyFilteringAsyncTask>::CreateTask(&Prerequisites).ConstructAndDispatchWhenReady(this, AsyncOperationMutex);
+	return TGraphTask<FTableTreeViewHierarchyFilteringAsyncTask>::CreateTask(&Prerequisites).ConstructAndDispatchWhenReady(this);
 }
 
 ////////////////////////////////////////////////////////////////////////////////////////////////////
@@ -1471,7 +1474,39 @@ void STableTreeView::TreeView_OnGetChildren(FTableTreeNodePtr InParent, TArray<F
 	if (InParent->OnLazyCreateChildren(SharedThis(this)))
 	{
 		// Node filtering does not apply to lazy expanded nodes.
-		// Grouping is ignored for lazy expanded nodes.
+		
+		// Apply grouping only if the node specifies the grouping that created the node.
+		// In this case, it will apply groupings after the identified one.
+		const FTreeNodeGrouping* AuthorGrouping = InParent->GetAuthorGrouping();
+		if (AuthorGrouping)
+		{
+			TArray<TSharedPtr<FTreeNodeGrouping>> Groupings;
+
+			bool bFound = false;
+			for (TSharedPtr<FTreeNodeGrouping> Grouping : CurrentGroupings)
+			{
+				if (bFound)
+				{
+					Groupings.Add(Grouping);
+				}
+				else
+				{
+					if (Grouping.Get() == AuthorGrouping)
+					{
+						bFound = true;
+					}
+				}
+			}
+
+			if (Groupings.Num() > 0)
+			{
+				// Extract children nodes.
+				TArray<FTableTreeNodePtr> ChildNodes;
+				InParent->SwapChildrenFast(reinterpret_cast<TArray<FBaseTreeNodePtr>&>(ChildNodes));
+
+				GroupNodesRec(ChildNodes, *InParent, 0, Groupings);
+			}
+		}
 
 		// Update aggregation.
 		UpdateAggregatedValuesRec(*InParent);
@@ -1664,7 +1699,7 @@ FGraphEventRef STableTreeView::StartGroupingTask(FGraphEventRef Prerequisite)
 		Prerequisites.Add(DispatchEvent);
 	}
 
-	return TGraphTask<FTableTreeViewGroupingAsyncTask>::CreateTask(&Prerequisites).ConstructAndDispatchWhenReady(this, &CurrentAsyncOpGroupings, AsyncOperationMutex);
+	return TGraphTask<FTableTreeViewGroupingAsyncTask>::CreateTask(&Prerequisites).ConstructAndDispatchWhenReady(this, &CurrentAsyncOpGroupings);
 }
 
 ////////////////////////////////////////////////////////////////////////////////////////////////////
@@ -2670,7 +2705,7 @@ FGraphEventRef STableTreeView::StartSortingTask(FGraphEventRef Prerequisite)
 		Prerequisites.Add(DispatchEvent);
 	}
 
-	return TGraphTask<FTableTreeViewSortingAsyncTask>::CreateTask(&Prerequisites).ConstructAndDispatchWhenReady(this, CurrentAsyncOpSorter, CurrentAsyncOpColumnSortMode, AsyncOperationMutex);
+	return TGraphTask<FTableTreeViewSortingAsyncTask>::CreateTask(&Prerequisites).ConstructAndDispatchWhenReady(this, CurrentAsyncOpSorter, CurrentAsyncOpColumnSortMode);
 }
 
 ////////////////////////////////////////////////////////////////////////////////////////////////////
@@ -3256,14 +3291,11 @@ void STableTreeView::OnClose()
 
 	bIsCloseScheduled = true;
 
-	if (bIsUpdateRunning && InProgressAsyncOperationEvent.IsValid() && !InProgressAsyncOperationEvent->IsComplete())
+	if (bIsUpdateRunning)
 	{
 		CancelCurrentAsyncOp();
-
-		FGraphEventArray Prerequisites;
-		Prerequisites.Add(InProgressAsyncOperationEvent);
-		AsyncCompleteTaskEvent = TGraphTask<FTableTreeViewAsyncCompleteTask>::CreateTask(&Prerequisites).ConstructAndDispatchWhenReady(SharedThis(this));
-		//TODO: FInsightsManager::Get()->AddInProgressAsyncOp(AsyncCompleteTaskEvent, TEXT("FTableTreeViewAsyncCompleteTask"));
+		check(InProgressAsyncOperationEvent.IsValid());
+		InProgressAsyncOperationEvent->Wait(ENamedThreads::GameThread);
 	}
 }
 
@@ -3313,6 +3345,11 @@ void STableTreeView::CancelCurrentAsyncOp()
 {
 	if (bIsUpdateRunning)
 	{
+		if (DispatchEvent.IsValid() && !DispatchEvent->IsComplete())
+		{
+			DispatchEvent->DispatchSubsequents();
+		}
+
 		AsyncOperationProgress.CancelCurrentAsyncOp();
 	}
 }
diff --git a/Engine/Plugins/Editor/AssetManagerEditor/Source/AssetManagerEditor/Private/Insights/Table/Widgets/STableTreeView.h b/Engine/Plugins/Editor/AssetManagerEditor/Source/AssetManagerEditor/Private/Insights/Table/Widgets/STableTreeView.h
--- a/Engine/Plugins/Editor/AssetManagerEditor/Source/AssetManagerEditor/Private/Insights/Table/Widgets/STableTreeView.h
+++ b/Engine/Plugins/Editor/AssetManagerEditor/Source/AssetManagerEditor/Private/Insights/Table/Widgets/STableTreeView.h
@@ -603,13 +603,6 @@ protected:
 	bool bIsUpdateRunning = false;
 	bool bIsCloseScheduled = false;
 
-	// The AsyncOperationLock is used a little differently than typic for a RW lock. 
-	// Instead of thinking of it as allowing multiple readers or one writer, it can be thought of
-	// as multiple shared lockers or one exclusive locker. The async operations all take read/shared locks
-	// and the destructor of the tree view takes an exclusive lock. This ensure that all async ops finish before
-	// the destructor runs and cleans up the data they are using
-	TSharedPtr<FRWLock> AsyncOperationMutex;
-
 	TArray<FTableTreeNodePtr> DummyGroupNodes;
 	FGraphEventRef InProgressAsyncOperationEvent;
 	FGraphEventRef AsyncCompleteTaskEvent;
@@ -636,45 +629,14 @@ protected:
 	FName LogListingName;
 };
 
-////////////////////////////////////////////////////////////////////////////////////////////////////
- 
-class FScopedExclusiveRWLock
-{
-public:
-	FScopedExclusiveRWLock(TSharedPtr<FRWLock>& InLock)
-		: Lock(InLock) {
-		Lock->WriteLock();
-	}
-
-	~FScopedExclusiveRWLock() { Lock->WriteUnlock(); }
-
-private:
-	TSharedPtr<FRWLock> Lock;
-	};
-
-class FScopedSharedRWLock
-{
-public:
-	FScopedSharedRWLock(TSharedPtr<FRWLock>& InLock)
-		: Lock(InLock) {
-		Lock->ReadLock();
-	}
-
-	~FScopedSharedRWLock() { Lock->ReadUnlock(); }
-
-private:
-	TSharedPtr<FRWLock> Lock;
-};
-
 ////////////////////////////////////////////////////////////////////////////////////////////////////
 
 class FTableTreeViewNodeFilteringAsyncTask
 {
 public:
-	FTableTreeViewNodeFilteringAsyncTask(STableTreeView* InPtr, TSharedPtr<FRWLock>& InAsyncOpLock)
+	FTableTreeViewNodeFilteringAsyncTask(STableTreeView* InPtr)
+		: TableTreeViewPtr(InPtr)
 	{
-		AsyncOpLock = InAsyncOpLock;
-		TableTreeViewPtr = InPtr;
 	}
 
 	FORCEINLINE TStatId GetStatId() const { RETURN_QUICK_DECLARE_CYCLE_STAT(FTableTreeViewNodeFilteringAsyncTask, STATGROUP_TaskGraphTasks); }
@@ -683,17 +645,14 @@ public:
 
 	void DoTask(ENamedThreads::Type CurrentThread, const FGraphEventRef& MyCompletionGraphEvent)
 	{
-		FScopedSharedRWLock SharedLock(AsyncOpLock);
 		if (TableTreeViewPtr)
 		{
 			TableTreeViewPtr->ApplyNodeFiltering();
 		}
-		AsyncOpLock.Reset();
 	}
 
 private:
 	STableTreeView* TableTreeViewPtr = nullptr;
-	TSharedPtr<FRWLock> AsyncOpLock;
 };
 
 ////////////////////////////////////////////////////////////////////////////////////////////////////
@@ -701,10 +660,9 @@ private:
 class FTableTreeViewHierarchyFilteringAsyncTask
 {
 public:
-	FTableTreeViewHierarchyFilteringAsyncTask(STableTreeView* InPtr, TSharedPtr<FRWLock> InAsyncOpLock)
+	FTableTreeViewHierarchyFilteringAsyncTask(STableTreeView* InPtr)
+		: TableTreeViewPtr(InPtr)
 	{
-		AsyncOpLock = InAsyncOpLock;
-		TableTreeViewPtr = InPtr;
 	}
 
 	FORCEINLINE TStatId GetStatId() const { RETURN_QUICK_DECLARE_CYCLE_STAT(FTableTreeViewHierarchyFilteringAsyncTask, STATGROUP_TaskGraphTasks); }
@@ -713,17 +671,14 @@ public:
 
 	void DoTask(ENamedThreads::Type CurrentThread, const FGraphEventRef& MyCompletionGraphEvent)
 	{
-		FScopedSharedRWLock ScopedLock(AsyncOpLock);
 		if (TableTreeViewPtr)
 		{
 			TableTreeViewPtr->ApplyHierarchyFiltering();
 		}
-		AsyncOpLock.Reset();
 	}
 
 private:
 	STableTreeView* TableTreeViewPtr = nullptr;
-	TSharedPtr<FRWLock> AsyncOpLock;
 };
 
 ////////////////////////////////////////////////////////////////////////////////////////////////////
@@ -731,12 +686,11 @@ private:
 class FTableTreeViewSortingAsyncTask
 {
 public:
-	FTableTreeViewSortingAsyncTask(STableTreeView* InPtr, ITableCellValueSorter* InSorter, EColumnSortMode::Type InColumnSortMode, TSharedPtr<FRWLock>& InAsyncOpLock)
+	FTableTreeViewSortingAsyncTask(STableTreeView* InPtr, ITableCellValueSorter* InSorter, EColumnSortMode::Type InColumnSortMode)
+		: TableTreeViewPtr(InPtr)
+		, Sorter(InSorter)
+		, ColumnSortMode(InColumnSortMode)
 	{
-		AsyncOpLock = InAsyncOpLock;
-		TableTreeViewPtr = InPtr;
-		Sorter = InSorter;
-		ColumnSortMode = InColumnSortMode;
 	}
 
 	FORCEINLINE TStatId GetStatId() const { RETURN_QUICK_DECLARE_CYCLE_STAT(FTableTreeViewSortingAsyncTask, STATGROUP_TaskGraphTasks); }
@@ -745,19 +699,16 @@ public:
 
 	void DoTask(ENamedThreads::Type CurrentThread, const FGraphEventRef& MyCompletionGraphEvent)
 	{
-		FScopedSharedRWLock ScopedLock(AsyncOpLock);
 		if (TableTreeViewPtr)
 		{
 			TableTreeViewPtr->SortTreeNodes(Sorter, ColumnSortMode);
 		}
-		AsyncOpLock.Reset();
 	}
 
 private:
-	STableTreeView* TableTreeViewPtr;
-	ITableCellValueSorter* Sorter;
-	EColumnSortMode::Type ColumnSortMode;
-	TSharedPtr<FRWLock> AsyncOpLock;
+	STableTreeView* TableTreeViewPtr = nullptr;
+	ITableCellValueSorter* Sorter = nullptr;
+	EColumnSortMode::Type ColumnSortMode = EColumnSortMode::Type::None;
 };
 
 ////////////////////////////////////////////////////////////////////////////////////////////////////
@@ -765,11 +716,10 @@ private:
 class FTableTreeViewGroupingAsyncTask
 {
 public:
-	FTableTreeViewGroupingAsyncTask(STableTreeView* InPtr, TArray<TSharedPtr<FTreeNodeGrouping>>* InGroupings, TSharedPtr<FRWLock>& InAsyncOpLock)
+	FTableTreeViewGroupingAsyncTask(STableTreeView* InPtr, TArray<TSharedPtr<FTreeNodeGrouping>>* InGroupings)
+		: TableTreeViewPtr(InPtr)
+		, Groupings(InGroupings)
 	{
-		AsyncOpLock = InAsyncOpLock;
-		TableTreeViewPtr = InPtr;
-		Groupings = InGroupings;
 	}
 
 	FORCEINLINE TStatId GetStatId() const { RETURN_QUICK_DECLARE_CYCLE_STAT(FTableTreeViewGroupingAsyncTask, STATGROUP_TaskGraphTasks); }
@@ -778,47 +728,15 @@ public:
 
 	void DoTask(ENamedThreads::Type CurrentThread, const FGraphEventRef& MyCompletionGraphEvent)
 	{
-		FScopedSharedRWLock ScopedLock(AsyncOpLock);
 		if (TableTreeViewPtr)
 		{
 			TableTreeViewPtr->CreateGroups(*Groupings);
 		}
-		AsyncOpLock.Reset();
 	}
 
 private:
 	STableTreeView* TableTreeViewPtr = nullptr;
-	TArray<TSharedPtr<FTreeNodeGrouping>>* Groupings;
-	TSharedPtr<FRWLock> AsyncOpLock;
-};
-
-////////////////////////////////////////////////////////////////////////////////////////////////////
-
-class FTableTreeViewAsyncCompleteTask
-{
-public:
-	FTableTreeViewAsyncCompleteTask(TSharedPtr<STableTreeView> InPtr)
-	{
-		TableTreeViewPtr = InPtr;
-	}
-
-	FORCEINLINE TStatId GetStatId() const { RETURN_QUICK_DECLARE_CYCLE_STAT(FTableTreeViewAsyncCompleteTask, STATGROUP_TaskGraphTasks); }
-	ENamedThreads::Type GetDesiredThread() { return ENamedThreads::Type::GameThread; }
-	static ESubsequentsMode::Type GetSubsequentsMode() { return ESubsequentsMode::TrackSubsequents; }
-
-	void DoTask(ENamedThreads::Type CurrentThread, const FGraphEventRef& MyCompletionGraphEvent)
-	{
-		// The role of this task is to keep the STableTreeView object alive until the task and it's prerequisites are completed and to destroy it on the game thread.
-		FGraphEventRef Event = TableTreeViewPtr->AsyncCompleteTaskEvent;
-		if (TableTreeViewPtr.IsValid())
-		{
-			TableTreeViewPtr.Reset();
-		}
-		//TODO: callback for FInsightsManager::Get()->RemoveInProgressAsyncOp(Event);
-	}
-
-private:
-	TSharedPtr<STableTreeView> TableTreeViewPtr;
+	TArray<TSharedPtr<FTreeNodeGrouping>>* Groupings = nullptr;
 };
 
 ////////////////////////////////////////////////////////////////////////////////////////////////////
@@ -826,11 +744,11 @@ private:
 class FSearchForItemToSelectTask
 {
 public:
-	FSearchForItemToSelectTask(TSharedPtr<FTableTaskCancellationToken> InToken, TSharedPtr<STableTreeView> InPtr, TSharedPtr<FRWLock> InAsyncOpLock)
+	FSearchForItemToSelectTask(TSharedPtr<FTableTaskCancellationToken> InToken, TSharedPtr<STableTreeView> InPtr)
 		: CancellationToken(InToken)
 		, TableTreeViewPtr(InPtr)
-		, AsyncOpLock(InAsyncOpLock)
-	{}
+	{
+	}
 
 	FORCEINLINE TStatId GetStatId() const { RETURN_QUICK_DECLARE_CYCLE_STAT(FSearchForItemToSelectTask, STATGROUP_TaskGraphTasks); }
 	ENamedThreads::Type GetDesiredThread() { return ENamedThreads::Type::AnyThread; }
@@ -838,18 +756,18 @@ public:
 
 	void DoTask(ENamedThreads::Type CurrentThread, const FGraphEventRef& MyCompletionGraphEvent)
 	{
-		FScopedSharedRWLock ScopedLock(AsyncOpLock);
-		if (TableTreeViewPtr)
+		if (!CancellationToken.IsValid() || !CancellationToken->ShouldCancel())
 		{
-			TableTreeViewPtr->SearchForItem(CancellationToken);
+			if (TableTreeViewPtr.IsValid())
+			{
+				TableTreeViewPtr->SearchForItem(CancellationToken);
+			}
 		}
-		AsyncOpLock.Reset();
 	}
 
 private:
 	TSharedPtr<FTableTaskCancellationToken> CancellationToken;
 	TSharedPtr<STableTreeView> TableTreeViewPtr;
-	TSharedPtr<FRWLock> AsyncOpLock;
 };
 
 ////////////////////////////////////////////////////////////////////////////////////////////////////
@@ -857,12 +775,12 @@ private:
 class FSelectNodeByTableRowIndexTask
 {
 public:
-	FSelectNodeByTableRowIndexTask(TSharedPtr<FTableTaskCancellationToken> InToken, TSharedPtr<STableTreeView> InPtr, uint32 InRowIndex, TSharedPtr<FRWLock> InAsyncOpLock)
+	FSelectNodeByTableRowIndexTask(TSharedPtr<FTableTaskCancellationToken> InToken, TSharedPtr<STableTreeView> InPtr, uint32 InRowIndex)
 		: CancellationToken(InToken)
 		, TableTreeViewPtr(InPtr)
-		, AsyncOpLock(InAsyncOpLock)
 		, RowIndex(InRowIndex)
-		{}
+	{
+	}
 
 	FORCEINLINE TStatId GetStatId() const { RETURN_QUICK_DECLARE_CYCLE_STAT(FSelectNodeByTableRowIndexTask, STATGROUP_TaskGraphTasks); }
 	ENamedThreads::Type GetDesiredThread() { return ENamedThreads::Type::GameThread; }
@@ -870,19 +788,19 @@ public:
 
 	void DoTask(ENamedThreads::Type CurrentThread, const FGraphEventRef& MyCompletionGraphEvent)
 	{
-		FScopedSharedRWLock ScopedLock(AsyncOpLock);
-		if (!CancellationToken->ShouldCancel() && TableTreeViewPtr)
+		if (!CancellationToken.IsValid() || !CancellationToken->ShouldCancel())
 		{
-			TableTreeViewPtr->SelectNodeByTableRowIndex(RowIndex);
+			if (TableTreeViewPtr.IsValid())
+			{
+				TableTreeViewPtr->SelectNodeByTableRowIndex(RowIndex);
+			}
 		}
-		AsyncOpLock.Reset();
 	}
 
 private:
 	TSharedPtr< FTableTaskCancellationToken> CancellationToken;
 	TSharedPtr<STableTreeView> TableTreeViewPtr;
-	TSharedPtr<FRWLock> AsyncOpLock;
-	uint32 RowIndex;
+	uint32 RowIndex = 0;
 };
 
 ////////////////////////////////////////////////////////////////////////////////////////////////////
diff --git a/Engine/Plugins/Editor/AssetManagerEditor/Source/AssetManagerEditor/Private/TreeView/AssetDependencyGrouping.cpp b/Engine/Plugins/Editor/AssetManagerEditor/Source/AssetManagerEditor/Private/TreeView/AssetDependencyGrouping.cpp
--- a/Engine/Plugins/Editor/AssetManagerEditor/Source/AssetManagerEditor/Private/TreeView/AssetDependencyGrouping.cpp
+++ b/Engine/Plugins/Editor/AssetManagerEditor/Source/AssetManagerEditor/Private/TreeView/AssetDependencyGrouping.cpp
@@ -167,6 +167,7 @@ void FPluginDependencyGrouping::GroupNodes(const TArray<UE::Insights::FTableTree
 		{
 			FName PluginAndDependenciesGroupName = AssetTable->GetNameForPlugin(PluginIndex);
 			TSharedPtr<FPluginAndDependenciesGroupNode> PluginAndDependenciesGroup = MakeShared<FPluginAndDependenciesGroupNode>(PluginAndDependenciesGroupName, AssetTable, PluginIndex);
+			PluginAndDependenciesGroup->SetAuthorGrouping(this);
 			FPluginSimpleGroupNode* PluginGroup = PluginAndDependenciesGroup->CreateChildren().Get();
 			PluginIndexToGroupNodeMap.Add(PluginIndex, PluginGroup);
 			RootPlugins->AddChildAndSetParent(PluginAndDependenciesGroup);
@@ -198,66 +199,6 @@ void FPluginDependencyGrouping::GroupNodes(const TArray<UE::Insights::FTableTree
 			continue;
 		}
 
-		// For each (visible) asset, we will create the following hierarchy:
-
-		// If the plugin does not have dependencies, a simple group will be created.
-		// |
-		// +-- [group:{PluginName}] (self) // FPluginSimpleGroupNode
-		// |   |
-		// |   +-- [asset:{AssetName1}]
-		// |   |
-		// |   +-- [asset:{AssetName2}]
-		// |   |
-		// |   ...
-
-		// By default, the "Plugin Dependencies" node is collapsed.
-		// |
-		// +-- [group:{PluginName}] (self + dependencies) // FPluginAndDependenciesGroupNode
-		// |   |
-		// |   +-- [group:Plugin Dependencies] (double click to expand) // FPluginDependenciesGroupNode, lazy
-		// |   |
-		// |   +-- [group:{PluginName}] (self) // FPluginSimpleGroupNode
-		// |       |
-		// |       +-- [asset:{AssetName1}]
-		// |       |
-		// |       +-- [asset:{AssetName2}]
-		// |       |
-		// |       ...
-
-		// When the "Plugin Dependencies" node is expanded, it will be populated with actual plugin dependencies.
-		// |
-		// +-- [group:{PluginName}] (self + dependencies) // FPluginAndDependenciesGroupNode
-		// |   |
-		// |   +-- [group:Plugin Dependencies] // FPluginDependenciesGroupNode expanded
-		// |   |   |
-		// |   |   +-- [group:{DependentPlugin1}] (self + dependencies) // FPluginAndDependenciesGroupNode
-		// |   |   |   |
-		// |   |   |   +-- [group:Plugin Dependencies] (double click to expand) // FPluginDependenciesGroupNode, lazy
-		// |   |   |   |
-		// |   |   |   +-- [group:{DependentPlugin1}] (self) // FPluginSimpleGroupNode
-		// |   |   |       |
-		// |   |   |       +-- [asset:{Asset1a}]
-		// |   |   |       |
-		// |   |   |       +-- [asset:{Asset1b}]
-		// |   |   |       ...
-		// |   |   |
-		// |   |   +-- [group:{DependentPlugin2}] (self, no further dependencies) // FPluginSimpleGroupNode
-		// |   |   |   |
-		// |   |   |   +-- [asset:{Asset2a}]
-		// |   |   |   |
-		// |   |   |   +-- [asset:{Asset2b}]
-		// |   |   |   ...
-		// |   |   |
-		// |   |   ...
-		// |   |
-		// |   +-- [group:{PluginName}] (self) // FPluginSimpleGroupNode
-		// |       |
-		// |       +-- [asset:{Asset1a}]
-		// |       |
-		// |       +-- [asset:{Asset1b}]
-		// |       |
-		// |       ...
-
 		FAssetTreeNode& AssetNode = NodePtr->As<FAssetTreeNode>();
 		const FAssetTableRow& Asset = AssetNode.GetAssetChecked();
 
@@ -273,9 +214,11 @@ void FPluginDependencyGrouping::GroupNodes(const TArray<UE::Insights::FTableTree
 		}
 		else
 		{
-			// Create the Plugin Self+Dependencies group node and add the current asset to the Self group.
+			// Create the Plugin group node (with "Assets" and "Dependencies" sub-groups)
+			// and add the current asset to the "Assets" sub-group.
 			FName PluginAndDependenciesGroupName = AssetTable->GetNameForPlugin(PluginIndex);
 			TSharedPtr<FPluginAndDependenciesGroupNode> PluginAndDependenciesGroup = MakeShared<FPluginAndDependenciesGroupNode>(PluginAndDependenciesGroupName, AssetTable, PluginIndex);
+			PluginAndDependenciesGroup->SetAuthorGrouping(this);
 			PluginGroup = PluginAndDependenciesGroup->CreateChildren().Get();
 			PluginGroup->AddChildAndSetParent(NodePtr);
 			PluginIndexToGroupNodeMap.Add(PluginIndex, PluginGroup);
diff --git a/Engine/Plugins/Editor/AssetManagerEditor/Source/AssetManagerEditor/Private/TreeView/AssetTreeNode.cpp b/Engine/Plugins/Editor/AssetManagerEditor/Source/AssetManagerEditor/Private/TreeView/AssetTreeNode.cpp
--- a/Engine/Plugins/Editor/AssetManagerEditor/Source/AssetManagerEditor/Private/TreeView/AssetTreeNode.cpp
+++ b/Engine/Plugins/Editor/AssetManagerEditor/Source/AssetManagerEditor/Private/TreeView/AssetTreeNode.cpp
@@ -53,6 +53,8 @@ const FSlateBrush* FAssetTreeNode::GetIcon(EStyle Style) const
 	}
 }
 
+////////////////////////////////////////////////////////////////////////////////////////////////////
+
 FLinearColor FAssetTreeNode::GetIconColor(EStyle Style) const
 {
 	switch (Style)
@@ -85,21 +87,29 @@ FLinearColor FAssetTreeNode::GetIconColor(EStyle Style) const
 	}
 }
 
+////////////////////////////////////////////////////////////////////////////////////////////////////
+
 FLinearColor FAssetTreeNode::GetColor(EStyle Style) const 
 {
 	return USlateThemeManager::Get().GetColor(EStyleColor::AccentWhite);
 }
 
+////////////////////////////////////////////////////////////////////////////////////////////////////
+
 FAssetTreeNode::EStyle FAssetTreeNode::GetStyle() const
 {
 	return IsGroup() ? EStyle::EGroup : (IsValidAsset() ? EStyle::EAsset : EStyle::EDefault);
 }
 
+////////////////////////////////////////////////////////////////////////////////////////////////////
+
 const FSlateBrush* FAssetTreeNode::GetIcon() const
 {
 	return GetIcon(GetStyle());
 }
 
+////////////////////////////////////////////////////////////////////////////////////////////////////
+
 FLinearColor FAssetTreeNode::GetIconColor() const
 {
 	return GetIconColor(GetStyle());
@@ -179,8 +189,7 @@ FAssetTreeNode::EStyle FPluginSimpleGroupNode::GetStyle() const
 {
 	return EStyle::EGroup;
 }
-////////////////////////////////////////////////////////////////////////////////////////////////////
-// FPluginSimpleGroupNode
+
 ////////////////////////////////////////////////////////////////////////////////////////////////////
 
 void FPluginSimpleGroupNode::AddAssetChildrenNodes()
@@ -214,29 +223,34 @@ FAssetTreeNode::EStyle FPluginAndDependenciesGroupNode::GetStyle() const
 	return EStyle::EPlugin;
 }
 
+////////////////////////////////////////////////////////////////////////////////////////////////////
+
 TSharedPtr<FPluginSimpleGroupNode> FPluginAndDependenciesGroupNode::CreateChildren()
 {
 	// [this]
 	// |
-	// +-- [group:Plugin Dependencies] (double click to expand) // FPluginDependenciesGroupNode, lazy
+	// +-- [group:Assets] // FPluginSimpleGroupNode
 	// |
-	// +-- [group:{PluginName}] (self) // FPluginSimpleGroupNode
+	// +-- [group:Dependencies] (double click to expand) // FPluginDependenciesGroupNode, lazy
 
 	FAssetTable& AssetTable = GetAssetTableChecked();
 	if (AssetTable.IsValidPluginIndex(PluginIndex))
 	{
+		// Create the Plugin Self group node (where asset nodes will be added).
+		//FName PluginGroupName = AssetTable.GetNameForPlugin(PluginIndex);
+		static FName PluginGroupName(TEXT("Assets"));
+		TSharedPtr<FPluginSimpleGroupNode> PluginGroup = MakeShared<FPluginSimpleGroupNode>(PluginGroupName, GetAssetTableWeak(), PluginIndex);
+		PluginGroup->SetAuthorGrouping(GetAuthorGrouping());
+		AddChildAndSetParent(PluginGroup);
+
 		// Create the Plugin Dependencies group node.
 		// The children nodes (list of dependent plugins) will be lazy created.
 		static FName DependenciesGroupName(TEXT("Dependencies"));
 		TSharedPtr<FPluginDependenciesGroupNode> DependenciesGroup = MakeShared<FPluginDependenciesGroupNode>(DependenciesGroupName, GetAssetTableWeak(), PluginIndex);
+		DependenciesGroup->SetAuthorGrouping(GetAuthorGrouping());
 		AddChildAndSetParent(DependenciesGroup);
 
-		// Create the Plugin Self group node (where asset nodes will be added).
-		//FName PluginGroupName = AssetTable.GetNameForPlugin(PluginIndex);
-		static FName PluginGroupName(TEXT("Assets"));
-		TSharedPtr<FPluginSimpleGroupNode> PluginGroup = MakeShared<FPluginSimpleGroupNode>(PluginGroupName, GetAssetTableWeak(), PluginIndex);
-		AddChildAndSetParent(PluginGroup);
-		return PluginGroup;
+		return PluginGroup; // the "Assets" group node
 	}
 	return SharedThis(this);
 }
@@ -245,6 +259,13 @@ TSharedPtr<FPluginSimpleGroupNode> FPluginAndDependenciesGroupNode::CreateChildr
 // FPluginDependenciesGroupNode
 ////////////////////////////////////////////////////////////////////////////////////////////////////
 
+FAssetTreeNode::EStyle FPluginDependenciesGroupNode::GetStyle() const
+{
+	return EStyle::EDependencies;
+}
+
+////////////////////////////////////////////////////////////////////////////////////////////////////
+
 const FText FPluginDependenciesGroupNode::GetExtraDisplayName() const
 {
 	if (!bAreChildrenCreated)
@@ -271,16 +292,17 @@ bool FPluginDependenciesGroupNode::OnLazyCreateChildren(TSharedPtr<class UE::Ins
 
 	// [this]
 	// |
-	// +-- [group:{DependentPlugin1}] (self + dependencies) // FPluginAndDependenciesGroupNode
+	// +-- [group:{DependentPlugin1}] // FPluginAndDependenciesGroupNode
 	// |   |
-	// |   +-- [group:Plugin Dependencies] (double click to expand) // FPluginDependenciesGroupNode, lazy
+	// |   +-- [group:Assets] // FPluginSimpleGroupNode
+	// |   +-- [group:Dependencies] (double click to expand) // FPluginDependenciesGroupNode, lazy
+	// |
+	// +-- [group:{DependentPlugin2}] // FPluginAndDependenciesGroupNode
 	// |   |
-	// |   +-- [group:{DependentPlugin1}] (self) // FPluginSimpleGroupNode
-	// |       |
-	// |       +-- [asset:{Asset1a}]
-	// |       |
-	// |       +-- [asset:{Asset1b}]
-	// |       ...
+	// |   +-- [group:Assets] // FPluginSimpleGroupNode
+	// |   +-- [group:Dependencies] (double click to expand) // FPluginDependenciesGroupNode, lazy
+	// |
+	// |   ...
 	// |
 
 	FAssetTable& AssetTable = GetAssetTableChecked();
@@ -296,6 +318,7 @@ bool FPluginDependenciesGroupNode::OnLazyCreateChildren(TSharedPtr<class UE::Ins
 				FName PluginGroupName = AssetTable.GetNameForPlugin(DependentPluginIndex);
 				const FAssetTablePluginInfo& DependentPluginInfo = AssetTable.GetPluginInfoByIndex(DependentPluginIndex);
 				TSharedPtr<FPluginAndDependenciesGroupNode> PluginGroup = MakeShared<FPluginAndDependenciesGroupNode>(PluginGroupName, GetAssetTableWeak(), DependentPluginIndex);
+				PluginGroup->SetAuthorGrouping(GetAuthorGrouping());
 				PluginGroup->CreateChildren()->AddAssetChildrenNodes();
 				AddChildAndSetParent(PluginGroup);
 			}
@@ -306,11 +329,6 @@ bool FPluginDependenciesGroupNode::OnLazyCreateChildren(TSharedPtr<class UE::Ins
 	return true;
 }
 
-FAssetTreeNode::EStyle FPluginDependenciesGroupNode::GetStyle() const
-{
-	return EStyle::EDependencies;
-}
-
 ////////////////////////////////////////////////////////////////////////////////////////////////////
 
 #undef LOCTEXT_NAMESPACE
diff --git a/Engine/Plugins/Editor/AssetManagerEditor/Source/AssetManagerEditor/Private/TreeView/AssetTreeNode.h b/Engine/Plugins/Editor/AssetManagerEditor/Source/AssetManagerEditor/Private/TreeView/AssetTreeNode.h
--- a/Engine/Plugins/Editor/AssetManagerEditor/Source/AssetManagerEditor/Private/TreeView/AssetTreeNode.h
+++ b/Engine/Plugins/Editor/AssetManagerEditor/Source/AssetManagerEditor/Private/TreeView/AssetTreeNode.h
@@ -8,6 +8,7 @@
 #include "Insights/Table/ViewModels/TableTreeNode.h"
 
 ////////////////////////////////////////////////////////////////////////////////////////////////////
+
 class FPluginSimpleGroupNode;
 
 class FAssetTreeNode : public UE::Insights::FTableTreeNode
@@ -47,7 +48,7 @@ public:
 	virtual const FSlateBrush* GetIcon() const final;
 	virtual FLinearColor GetIconColor() const final;
 	virtual FLinearColor GetColor() const final;
-	
+
 	bool IsPluginRelatedNode() const { return Is<FPluginSimpleGroupNode>(); }
 	bool IsAssetRelatedNode() const { return !IsPluginRelatedNode(); }
 
@@ -95,6 +96,7 @@ public:
 	explicit FAssetDependenciesGroupTreeNode(const FName InGroupName, TWeakPtr<FAssetTable> InParentTable, int32 InParentRowIndex)
 		: FAssetTreeNode(InGroupName, InParentTable, InParentRowIndex, true)
 		, bAreChildrenCreated(false)
+		, AuthorGrouping(nullptr)
 	{
 		// Initially collapsed. Lazy create children when first expanded.
 		SetExpansion(false);
@@ -103,12 +105,16 @@ public:
 	virtual ~FAssetDependenciesGroupTreeNode() {}
 
 	virtual EStyle GetStyle() const override;
-	
+
 	virtual const FText GetExtraDisplayName() const;
 	virtual bool OnLazyCreateChildren(TSharedPtr<class UE::Insights::STableTreeView> InTableTreeView) override;
 
+	virtual const UE::Insights::FTreeNodeGrouping* GetAuthorGrouping() const override { return AuthorGrouping; }
+	void SetAuthorGrouping(const UE::Insights::FTreeNodeGrouping* InAuthorGrouping) { AuthorGrouping = InAuthorGrouping; }
+
 private:
 	bool bAreChildrenCreated;
+	const UE::Insights::FTreeNodeGrouping* AuthorGrouping;
 };
 
 /** Type definition for shared pointers to instances of FAssetDependenciesGroupTreeNode. */
@@ -125,6 +131,7 @@ public:
 	explicit FPluginSimpleGroupNode(const FName InGroupName, TWeakPtr<FAssetTable> InParentTable, int32 InPluginIndex)
 		: FAssetTreeNode(InGroupName, InParentTable)
 		, PluginIndex(InPluginIndex)
+		, AuthorGrouping(nullptr)
 	{
 	}
 
@@ -136,8 +143,12 @@ public:
 
 	int32 GetPluginIndex() const { return PluginIndex; }
 
+	virtual const UE::Insights::FTreeNodeGrouping* GetAuthorGrouping() const override { return AuthorGrouping; }
+	void SetAuthorGrouping(const UE::Insights::FTreeNodeGrouping* InAuthorGrouping) { AuthorGrouping = InAuthorGrouping; }
+
 protected:
 	int32 PluginIndex;
+	const UE::Insights::FTreeNodeGrouping* AuthorGrouping;
 };
 
 ////////////////////////////////////////////////////////////////////////////////////////////////////
@@ -153,11 +164,11 @@ public:
 	{
 	}
 
-	TSharedPtr<FPluginSimpleGroupNode> CreateChildren();
+	virtual ~FPluginAndDependenciesGroupNode() {}
 
 	virtual EStyle GetStyle() const override;
 
-	virtual ~FPluginAndDependenciesGroupNode() {}
+	TSharedPtr<FPluginSimpleGroupNode> CreateChildren();
 };
 
 ////////////////////////////////////////////////////////////////////////////////////////////////////
@@ -176,11 +187,11 @@ public:
 
 	virtual ~FPluginDependenciesGroupNode() {}
 
+	virtual EStyle GetStyle() const override;
+
 	virtual const FText GetExtraDisplayName() const;
 	virtual bool OnLazyCreateChildren(TSharedPtr<class UE::Insights::STableTreeView> InTableTreeView) override;
 
-	virtual EStyle GetStyle() const override;
-
 private:
 	bool bAreChildrenCreated;
 };
