AuditDiskSizeTool -- Updating 'Unique Size' column to work for plugin dependency sizes.

diff --git a/Engine/Plugins/Editor/AssetManagerEditor/Source/AssetManagerEditor/Private/TreeView/AssetTable.cpp b/Engine/Plugins/Editor/AssetManagerEditor/Source/AssetManagerEditor/Private/TreeView/AssetTable.cpp
--- a/Engine/Plugins/Editor/AssetManagerEditor/Source/AssetManagerEditor/Private/TreeView/AssetTable.cpp
+++ b/Engine/Plugins/Editor/AssetManagerEditor/Source/AssetManagerEditor/Private/TreeView/AssetTable.cpp
@@ -968,7 +968,7 @@ void FAssetTable::AddDefaultColumns()
 
 		Column.SetShortName(LOCTEXT("TotalSizeUniqueDependenciesColumnName", "Unique"));
 		Column.SetTitleName(LOCTEXT("TotalSizeUniqueDependenciesColumnTitle", "Total Unique Dependency Size"));
-		Column.SetDescription(LOCTEXT("TotalSizeUniqueDependenciesColumnIdDesc", "Sum of the staged compressed sizes of all dependencies of this asset, counted only once"));
+		Column.SetDescription(LOCTEXT("TotalSizeUniqueDependenciesColumnIdDesc", "Sum of the staged compressed sizes of all dependencies of this item, counted only once"));
 
 		Column.SetFlags(ETableColumnFlags::CanBeHidden | ETableColumnFlags::CanBeFiltered);
 
@@ -982,7 +982,16 @@ void FAssetTable::AddDefaultColumns()
 		public:
 			virtual const TOptional<FTableCellValue> GetValue(const FTableColumn& Column, const FBaseTreeNode& Node) const override
 			{
-				if (Node.IsGroup())
+				if (Node.Is<FPluginSimpleGroupNode>() && !Node.Is<FPluginDependenciesGroupNode>())
+				{
+					// This is node represents a single plugin (it might be the plugin itself or the plugin+deps node for that plugin)
+					const FPluginSimpleGroupNode& PluginNode = Node.As<FPluginSimpleGroupNode>();
+					TSharedPtr<FTable> TablePtr = PluginNode.GetParentTable().Pin();
+					const FAssetTable& AssetTable = static_cast<const FAssetTable&>(*TablePtr);
+					const FAssetTablePluginInfo& PluginInfo = AssetTable.GetPluginInfoByIndexChecked(PluginNode.GetPluginIndex());
+					return FTableCellValue(PluginInfo.GetOrComputeTotalSizeUniqueDependencies(static_cast<const FAssetTable&>(*TablePtr)));
+				}
+				else if (Node.IsGroup())
 				{
 					const FTableTreeNode& NodePtr = static_cast<const FTableTreeNode&>(Node);
 					if (NodePtr.GetRowId().HasValidIndex())
