Add search support to the Reference Viewer collection filter dropdown

Rewrite the Reference Viewer's collection filter UI around an SComboButton instead of an SComboBox. Construct the SComboButton's dropdown list using an FMenuBuilder to add support for searching by typing.

diff --git a/Engine/Plugins/Editor/AssetManagerEditor/Source/AssetManagerEditor/Private/ReferenceViewer/SReferenceViewer.cpp b/Engine/Plugins/Editor/AssetManagerEditor/Source/AssetManagerEditor/Private/ReferenceViewer/SReferenceViewer.cpp
--- a/Engine/Plugins/Editor/AssetManagerEditor/Source/AssetManagerEditor/Private/ReferenceViewer/SReferenceViewer.cpp
+++ b/Engine/Plugins/Editor/AssetManagerEditor/Source/AssetManagerEditor/Private/ReferenceViewer/SReferenceViewer.cpp
@@ -459,15 +459,13 @@ void SReferenceViewer::Construct(const FArguments& InArgs)
 							SNew(SBox)
 							.WidthOverride(100)
 							[
-								SAssignNew(CollectionsCombo, SComboBox<TSharedPtr<FName>>)
-								.OptionsSource(&CollectionsComboList)
-								.OnComboBoxOpening(this, &SReferenceViewer::UpdateCollectionsComboList)
-								.OnGenerateWidget(this, &SReferenceViewer::GenerateCollectionFilterItem)
-								.OnSelectionChanged(this, &SReferenceViewer::HandleCollectionFilterChanged)
-								.ToolTipText(this, &SReferenceViewer::GetCollectionFilterText)
+								SNew(SComboButton)
+								.OnGetMenuContent(this, &SReferenceViewer::BuildCollectionFilterMenu)
+								.ButtonContent()
 								[
 									SNew(STextBlock)
-									.Text(this, &SReferenceViewer::GetCollectionFilterText)
+									.Text(this, &SReferenceViewer::GetCollectionComboButtonText)
+									.ToolTipText(this, &SReferenceViewer::GetCollectionComboButtonText)
 								]
 							]
 						]
@@ -554,8 +552,6 @@ void SReferenceViewer::Construct(const FArguments& InArgs)
 		]
 	];
 
-	UpdateCollectionsComboList();
-
 	SetCanTick(true);
 }
 
@@ -1011,19 +1007,6 @@ ECheckBoxState SReferenceViewer::IsSearchBreadthEnabledChecked() const
 	return Settings->IsSearchBreadthLimited() ? ECheckBoxState::Checked : ECheckBoxState::Unchecked;
 }
 
-TSharedRef<SWidget> SReferenceViewer::GenerateCollectionFilterItem(TSharedPtr<FName> InItem)
-{
-	FText ItemAsText = FText::FromName(*InItem);
-	return
-		SNew(SBox)
-		.WidthOverride(300)
-		[
-			SNew(STextBlock)
-			.Text(ItemAsText)
-			.ToolTipText(ItemAsText)
-		];
-}
-
 void SReferenceViewer::OnEnableCollectionFilterChanged(ECheckBoxState NewState)
 {
 	const bool bNewValue = NewState == ECheckBoxState::Checked;
@@ -1040,8 +1023,63 @@ ECheckBoxState SReferenceViewer::IsEnableCollectionFilterChecked() const
 	return Settings->GetEnableCollectionFilter() ? ECheckBoxState::Checked : ECheckBoxState::Unchecked;
 }
 
-void SReferenceViewer::UpdateCollectionsComboList()
+FText SReferenceViewer::GetCollectionComboButtonText() const
+{
+	return FText::FromName(GraphObj->GetCurrentCollectionFilter());
+}
+
+void SReferenceViewer::CollectionFilterAddMenuEntry(FMenuBuilder& MenuBuilder, const FName& CollectionName)
+{
+	FExecuteAction ActionClicked = FExecuteAction::CreateLambda([this, CollectionName]()
+	{
+		// Make sure collection filtering is enabled now that the user clicked something in the menu.
+		Settings->SetEnableCollectionFilter(true);
+
+		FName CurrentCollectionFilter = GraphObj->GetCurrentCollectionFilter();
+
+		// Update the filter and rebuild the graph if the filter changed.
+		if (CurrentCollectionFilter != CollectionName)
+		{
+			GraphObj->SetCurrentCollectionFilter(CollectionName);
+			RebuildGraph();
+		}
+	});
+
+	FIsActionChecked ActionChecked = FIsActionChecked::CreateLambda([this, CollectionName]() -> bool
+	{
+		return GraphObj->GetCurrentCollectionFilter() == CollectionName;
+	});
+
+	MenuBuilder.AddMenuEntry(
+		FText::FromName(CollectionName),
+		FText::FromName(CollectionName),
+		FSlateIcon(),
+		FUIAction(ActionClicked, FCanExecuteAction(), ActionChecked),
+		NAME_None, //InExtensionHook
+		EUserInterfaceActionType::RadioButton
+		);
+}
+
+TSharedRef<SWidget> SReferenceViewer::BuildCollectionFilterMenu()
 {
+	FMenuBuilder MenuBuilder(true /* Pass true to close dropdown after selection. */, nullptr);
+
+	MenuBuilder.AddMenuEntry(
+		LOCTEXT("CollectionFilterSelectNone", "Select None"),
+		LOCTEXT("CollectionFilterSelectNoCollection", "Select no collection."),
+		FSlateIcon(),
+		FExecuteAction::CreateLambda([this]()
+		{
+			// Make sure collection filtering is enabled.
+			Settings->SetEnableCollectionFilter(true);
+
+			GraphObj->SetCurrentCollectionFilter(NAME_None);
+			RebuildGraph();
+		})
+	);
+
+	MenuBuilder.AddSeparator();
+
 	TArray<FName> CollectionNames;
 	{
 		FCollectionManagerModule& CollectionManagerModule = FCollectionManagerModule::GetModule();
@@ -1059,61 +1097,16 @@ void SReferenceViewer::UpdateCollectionsComboList()
 				CollectionNames.AddUnique(Collection.Name);
 			}
 		}
-	}
-	CollectionNames.Sort([](const FName& A, const FName& B) { return A.Compare(B) < 0; });
-
-	CollectionsComboList.Reset();
-	CollectionsComboList.Add(MakeShared<FName>(NAME_None));
-	for (FName CollectionName : CollectionNames)
-	{
-		CollectionsComboList.Add(MakeShared<FName>(CollectionName));
-	}
-
-	if (CollectionsCombo)
-	{
-		CollectionsCombo->ClearSelection();
-		CollectionsCombo->RefreshOptions();
-
-		if (GraphObj)
-		{
-			const FName CurrentFilter = GraphObj->GetCurrentCollectionFilter();
-
-			const int32 SelectedItemIndex = CollectionsComboList.IndexOfByPredicate([CurrentFilter](const TSharedPtr<FName>& InItem)
-			{
-				return CurrentFilter == *InItem;
-			});
 
-			if (SelectedItemIndex != INDEX_NONE)
-			{
-				CollectionsCombo->SetSelectedItem(CollectionsComboList[SelectedItemIndex]);
-			}
-		}
+		CollectionNames.Sort([](const FName& A, const FName& B) { return A.Compare(B) < 0; });
 	}
-}
 
-void SReferenceViewer::HandleCollectionFilterChanged(TSharedPtr<FName> Item, ESelectInfo::Type SelectInfo)
-{
-	if (GraphObj && Item)
+	for (const FName& CollectionName : CollectionNames)
 	{
-		const FName NewFilter = *Item;
-		const FName CurrentFilter = GraphObj->GetCurrentCollectionFilter();
-		if (CurrentFilter != NewFilter)
-		{
-			if (CurrentFilter == NAME_None)
-			{
-				// Automatically check the box to enable the filter if the previous filter was None
-				Settings->SetEnableCollectionFilter(true);
-			}
-
-			GraphObj->SetCurrentCollectionFilter(NewFilter);
-			RebuildGraph();
-		}
+		CollectionFilterAddMenuEntry(MenuBuilder, CollectionName);
 	}
-}
 
-FText SReferenceViewer::GetCollectionFilterText() const
-{
-	return FText::FromName(GraphObj->GetCurrentCollectionFilter());
+	return MenuBuilder.MakeWidget();
 }
 
 void SReferenceViewer::OnEnablePluginFilterChanged(ECheckBoxState NewState)
diff --git a/Engine/Plugins/Editor/AssetManagerEditor/Source/AssetManagerEditor/Private/ReferenceViewer/SReferenceViewer.h b/Engine/Plugins/Editor/AssetManagerEditor/Source/AssetManagerEditor/Private/ReferenceViewer/SReferenceViewer.h
--- a/Engine/Plugins/Editor/AssetManagerEditor/Source/AssetManagerEditor/Private/ReferenceViewer/SReferenceViewer.h
+++ b/Engine/Plugins/Editor/AssetManagerEditor/Source/AssetManagerEditor/Private/ReferenceViewer/SReferenceViewer.h
@@ -123,10 +123,9 @@ private:
 
 	void OnEnableCollectionFilterChanged(ECheckBoxState NewState);
 	ECheckBoxState IsEnableCollectionFilterChecked() const;
-	TSharedRef<SWidget> GenerateCollectionFilterItem(TSharedPtr<FName> InItem);
-	void UpdateCollectionsComboList();
-	void HandleCollectionFilterChanged(TSharedPtr<FName> Item, ESelectInfo::Type SelectInfo);
-	FText GetCollectionFilterText() const;
+	void CollectionFilterAddMenuEntry(FMenuBuilder& MenuBuilder, const FName& CollectionName);
+	TSharedRef<SWidget> BuildCollectionFilterMenu();
+	FText GetCollectionComboButtonText() const;
 
 	void OnEnablePluginFilterChanged(ECheckBoxState NewState);
 	ECheckBoxState IsEnablePluginFilterChecked() const;
