Improve the Reference Viewer

* Fix a bug where only the plugins of first-level references/dependencies were listed in the plugin filter dropdown.
* Polish the plugin filter dropdown to avoid adding a separator between plugins and non-plugins if there are only non-plugins in the list.
* Similarly, avoid adding a separator below "Select All" and "Select None" if there are no plugins in the list (for example because the graph is empty).
* Remove an unnecessary member variable.

diff --git a/Engine/Plugins/Editor/AssetManagerEditor/Source/AssetManagerEditor/Private/ReferenceViewer/EdGraph_ReferenceViewer.cpp b/Engine/Plugins/Editor/AssetManagerEditor/Source/AssetManagerEditor/Private/ReferenceViewer/EdGraph_ReferenceViewer.cpp
--- a/Engine/Plugins/Editor/AssetManagerEditor/Source/AssetManagerEditor/Private/ReferenceViewer/EdGraph_ReferenceViewer.cpp
+++ b/Engine/Plugins/Editor/AssetManagerEditor/Source/AssetManagerEditor/Private/ReferenceViewer/EdGraph_ReferenceViewer.cpp
@@ -659,11 +659,6 @@ void UEdGraph_ReferenceViewer::GetUnfilteredGraphPluginNamesRecursive(bool bRefe
 	{
 		return;
 	}
-
-	if (OutAssetIdentifiers.Contains(InAssetIdentifier))
-	{
-		return;
-	}
 	
 	IAssetRegistry& AssetRegistry = FModuleManager::LoadModuleChecked<FAssetRegistryModule>("AssetRegistry").Get();
 	TArray<FAssetDependency> LinksToAsset;
@@ -678,12 +673,18 @@ void UEdGraph_ReferenceViewer::GetUnfilteredGraphPluginNamesRecursive(bool bRefe
 	
 	for (const FAssetDependency& Link : LinksToAsset)
 	{
+		// Avoid loops by skipping assets we've already visited.
+		if (OutAssetIdentifiers.Contains(Link.AssetId))
+		{
+			continue;;
+		}
+
 		// Don't add assets that will be hidden by Reference Viewer settings the user cannot change.
 		if (!IsPackageIdentifierPassingFilter(Link.AssetId))
 		{
 			continue;
 		}
-		
+
 		OutAssetIdentifiers.Add(Link.AssetId);
 
 		GetUnfilteredGraphPluginNamesRecursive(bReferencers, Link.AssetId, InCurrentDepth + 1, InMaxDepth, Query, OutAssetIdentifiers);
diff --git a/Engine/Plugins/Editor/AssetManagerEditor/Source/AssetManagerEditor/Private/ReferenceViewer/SReferenceViewer.cpp b/Engine/Plugins/Editor/AssetManagerEditor/Source/AssetManagerEditor/Private/ReferenceViewer/SReferenceViewer.cpp
--- a/Engine/Plugins/Editor/AssetManagerEditor/Source/AssetManagerEditor/Private/ReferenceViewer/SReferenceViewer.cpp
+++ b/Engine/Plugins/Editor/AssetManagerEditor/Source/AssetManagerEditor/Private/ReferenceViewer/SReferenceViewer.cpp
@@ -504,7 +504,7 @@ void SReferenceViewer::Construct(const FArguments& InArgs)
 							SNew(SBox)
 							.WidthOverride(100)
 							[
-								SAssignNew(PluginsComboButton, SComboButton)
+								SNew(SComboButton)
 								.OnGetMenuContent(this, &SReferenceViewer::BuildPluginFilterMenu)
 								.ButtonContent()
 								[
@@ -1200,8 +1200,6 @@ TSharedRef<SWidget> SReferenceViewer::BuildPluginFilterMenu()
 		})
 	);
 
-	MenuBuilder.AddSeparator();
-	
 	const TArray<FName> InitialPluginFilter = GraphObj->GetCurrentPluginFilter();
 
 	TArray<FName> EnabledPluginNames;
@@ -1214,7 +1212,12 @@ TSharedRef<SWidget> SReferenceViewer::BuildPluginFilterMenu()
 	TArray<FName> PluginNames = GraphObj->GetEncounteredPluginsAmongNodes();
 	PluginNames.Sort([](const FName& A, const FName& B) { return A.Compare(B) < 0; });
 
-	bool bAddedNonPlugins = false;
+	if (PluginNames.Num() > 0)
+	{
+		MenuBuilder.AddSeparator();
+	}
+
+	int NumNonPluginsAdded = 0;
 	// First add any "plugins" that are not actual plugins, such as /Game and /Engine.
 	for (const FName& PluginName : PluginNames)
 	{
@@ -1222,14 +1225,15 @@ TSharedRef<SWidget> SReferenceViewer::BuildPluginFilterMenu()
 		{
 			continue;
 		}
-
-		bAddedNonPlugins = true;
 		
 		PluginFilterAddMenuEntry(MenuBuilder, PluginName);
+
+		++NumNonPluginsAdded;
 	}
 
-	// Separate non-plugins from real plugins.
-	if (bAddedNonPlugins)
+	// Add a separator between non-real plugins (such as /Game and /Engine) and actual plugins.
+	// Skip this if no non-real plugins were added or if there are no more plugins to add.
+	if (NumNonPluginsAdded > 0 && NumNonPluginsAdded < PluginNames.Num())
 	{
 		MenuBuilder.AddSeparator();
 	}
diff --git a/Engine/Plugins/Editor/AssetManagerEditor/Source/AssetManagerEditor/Private/ReferenceViewer/SReferenceViewer.h b/Engine/Plugins/Editor/AssetManagerEditor/Source/AssetManagerEditor/Private/ReferenceViewer/SReferenceViewer.h
--- a/Engine/Plugins/Editor/AssetManagerEditor/Source/AssetManagerEditor/Private/ReferenceViewer/SReferenceViewer.h
+++ b/Engine/Plugins/Editor/AssetManagerEditor/Source/AssetManagerEditor/Private/ReferenceViewer/SReferenceViewer.h
@@ -237,9 +237,6 @@ private:
 	/** List of collection filter options */
 	TArray<TSharedPtr<FName>> CollectionsComboList;
 
-	/** Combo box for plugin filter options */
-	TSharedPtr<SComboButton> PluginsComboButton;
-
 	/**
 	 * Whether to visually show to the user the option of "Search Depth Limit" or hide it and fix it to a default value:
 	 * - If 0 or negative, it will show to the user the option of "Search Depth Limit".
