SourceControl: Fixed `FAssetSourceControlContextMenuState::ExecuteSCCRevertWritable()` no longer respecting the read only state.

The original changelist, #22893640, had the source control operations work on the `PackageFileNamesToRevert` container. However, somewhere since then, the code was refactored and left the `PackageFileNamesToRevert` container unused.

#rnx

diff --git a/Engine/Plugins/Editor/AssetManagerEditor/Source/AssetManagerEditor/Private/AssetSourceControlContextMenu.cpp b/Engine/Plugins/Editor/AssetManagerEditor/Source/AssetManagerEditor/Private/AssetSourceControlContextMenu.cpp
--- a/Engine/Plugins/Editor/AssetManagerEditor/Source/AssetManagerEditor/Private/AssetSourceControlContextMenu.cpp
+++ b/Engine/Plugins/Editor/AssetManagerEditor/Source/AssetManagerEditor/Private/AssetSourceControlContextMenu.cpp
@@ -854,7 +854,7 @@ void FAssetSourceControlContextMenuState::ExecuteSCCRevertWritable() const
 		return true;
 	};
 
-	SourceControlHelpers::ApplyOperationAndReloadPackages(PackageNames, RevertOperation);
+	SourceControlHelpers::ApplyOperationAndReloadPackages(PackageFileNamesToRevert, RevertOperation);
 
 	// Tell UncontrolledCL to refresh to pick up changes to files it was watching.
 	FUncontrolledChangelistsModule& UncontrolledChangelistsModule = FUncontrolledChangelistsModule::Get();
