Added support for shader pseudo asset data in the asset disk size tool.

diff --git a/Engine/Plugins/Editor/AssetManagerEditor/Source/AssetManagerEditor/Private/TreeView/SAssetTableTreeView.cpp b/Engine/Plugins/Editor/AssetManagerEditor/Source/AssetManagerEditor/Private/TreeView/SAssetTableTreeView.cpp
--- a/Engine/Plugins/Editor/AssetManagerEditor/Source/AssetManagerEditor/Private/TreeView/SAssetTableTreeView.cpp
+++ b/Engine/Plugins/Editor/AssetManagerEditor/Source/AssetManagerEditor/Private/TreeView/SAssetTableTreeView.cpp
@@ -1749,12 +1749,11 @@ void SAssetTableTreeView::RefreshAssets()
 
 	TMap<const TCHAR*, DeprecatedTCharSetType, FDefaultSetAllocator, TStringPointerMapKeyFuncs_DEPRECATED<const TCHAR*, DeprecatedTCharSetType>> DiscoveredPluginDependencyEdges;
 
+	TMap<FAssetData, int32> AssetToIndexMap;
+	TArray<FAssetData> SourceAssets;
+
 	if (IsRegistrySourceValid())
 	{
-		TMap<FAssetData, int32> AssetToIndexMap;
-
-		TArray<FAssetData> SourceAssets;
-
 		RegistrySource.GetOwnedRegistryState()->GetAllAssets(TSet<FName>(), SourceAssets);
 
 		for (int32 SourceAssetIndex = 0; SourceAssetIndex < SourceAssets.Num(); SourceAssetIndex++)
@@ -1898,70 +1897,137 @@ void SAssetTableTreeView::RefreshAssets()
 
 	if (CookMetadata.IsValid())
 	{
-		const UE::Cook::FCookMetadataPluginHierarchy& PluginHierarchy = CookMetadata.GetPluginHierarchy();
-		int64 TotalMismatch = 0;
-		int64 TotalSize = 0;
-		for (const UE::Cook::FCookMetadataPluginEntry& PluginEntry : PluginHierarchy.PluginsEnabledAtCook)
+		// Shader data
 		{
-			const TCHAR* StoredPluginName = AssetTable->StoreStr(PluginEntry.Name);
-			FAssetTablePluginInfo& PluginInfo = AssetTable->GetOrCreatePluginInfo(StoredPluginName);
-			if (CookMetadata.GetSizesPresent() != UE::Cook::ECookMetadataSizesPresent::NotPresent)
+			const UE::Cook::FCookMetadataShaderPseudoHierarchy& ShaderHierarchy = CookMetadata.GetShaderPseudoHierarchy();
+
+			const TCHAR* ShaderAssetType = AssetTable->StoreStr(TEXT("Shader Pseudo Asset"));
+			const uint32 AssetTypeHash = GetTypeHash(FStringView(ShaderAssetType));
+			FLinearColor ShaderColor = USlateThemeManager::Get().GetColor((EStyleColor)((uint32)EStyleColor::AccentBlue + AssetTypeHash % ((uint32)EStyleColor::AccentGreen - (uint32)EStyleColor::AccentBlue)));
+
+			uint64 ShaderStartIndexInAssetTable = AssetTable->GetAssets().Num();
+
+			// Build all the additional asset rows for the ShaderAssets.
+			for (const UE::Cook::FCookMetadataShaderPseudoAsset& ShaderAsset : ShaderHierarchy.ShaderAssets)
 			{
-				PluginInfo.Size = PluginEntry.ExclusiveSizes[UE::Cook::EPluginSizeTypes::Installed];
+				FAssetTableRow AssetRow;
+				{
+					int FirstSlashIndex = INDEX_NONE;
+					ShaderAsset.Name.FindChar(TEXT('/'), FirstSlashIndex);
+					FString PluginName;
+					if (FirstSlashIndex != INDEX_NONE)
+					{
+						PluginName = ShaderAsset.Name.Left(FirstSlashIndex);
+					}
+					AssetRow.PluginName = AssetTable->StoreStr(PluginName);
+
+					int LastSlashIndex = INDEX_NONE;
+					ShaderAsset.Name.FindLastChar(TEXT('/'), LastSlashIndex);
+					FString AssetName;
+					if (FirstSlashIndex != INDEX_NONE)
+					{
+						AssetName = ShaderAsset.Name.Right(ShaderAsset.Name.Len() - LastSlashIndex - 1);
+					}
+					AssetRow.Name = AssetTable->StoreStr(AssetName);
+				}
+				AssetRow.Type = ShaderAssetType;
+				AssetRow.Color = ShaderColor;
+				AssetRow.StagedCompressedSizeRequiredInstall = ShaderAsset.CompressedSize;
+				AssetTable->AddAsset(AssetRow);
 			}
-			else if (const int64* SizePtr = PluginToSizeMap.Find(PluginEntry.Name))
+			
+			// This is not efficient as most packages don't have shader dependencies at all
+			for (const FAssetData& SourceAsset : SourceAssets)
 			{
-				PluginInfo.Size = *SizePtr;
+				if (const TPair<int32, int32>* ShaderDependencyRange = ShaderHierarchy.PackageShaderDependencyMap.Find(SourceAsset.PackageName))
+				{
+					if (int32* AssetIndex = AssetToIndexMap.Find(SourceAsset))
+					{
+						FAssetTableRow* OwningAssetRow = AssetTable->GetAsset(*AssetIndex);
+
+						for (int32 ShaderPseudoAssetIndex = ShaderDependencyRange->Key; ShaderPseudoAssetIndex < ShaderDependencyRange->Value; ShaderPseudoAssetIndex++)
+						{
+							int64 ShaderIndexInAssetTable = ShaderStartIndexInAssetTable + ShaderHierarchy.DependencyList[ShaderPseudoAssetIndex];
+							if (ensure(ShaderIndexInAssetTable < AssetTable->GetAssets().Num()))
+							{
+								OwningAssetRow->Dependencies.AddUnique(ShaderIndexInAssetTable);
+								AssetTable->GetAsset(ShaderIndexInAssetTable)->Referencers.AddUnique(*AssetIndex);
+							}
+						}
+					}
+				}
 			}
 
-			///
-			/// validation
-			///
-			if (const int64* SizePtr = PluginToSizeMap.Find(PluginEntry.Name))
+			AssetTable->SetVisibleAssetCount(AssetTable->GetAssets().Num());
+		}
+		
+		// Plugin Data
+		{
+			const UE::Cook::FCookMetadataPluginHierarchy& PluginHierarchy = CookMetadata.GetPluginHierarchy();
+			int64 TotalMismatch = 0;
+			int64 TotalSize = 0;
+			for (const UE::Cook::FCookMetadataPluginEntry& PluginEntry : PluginHierarchy.PluginsEnabledAtCook)
 			{
+				const TCHAR* StoredPluginName = AssetTable->StoreStr(PluginEntry.Name);
+				FAssetTablePluginInfo& PluginInfo = AssetTable->GetOrCreatePluginInfo(StoredPluginName);
 				if (CookMetadata.GetSizesPresent() != UE::Cook::ECookMetadataSizesPresent::NotPresent)
 				{
-					int64 TotalSizeOfPluginInMetadata = PluginEntry.ExclusiveSizes[UE::Cook::EPluginSizeTypes::Installed];
-					if (*SizePtr != TotalSizeOfPluginInMetadata)
+					PluginInfo.Size = PluginEntry.ExclusiveSizes[UE::Cook::EPluginSizeTypes::Installed];
+				}
+				else if (const int64* SizePtr = PluginToSizeMap.Find(PluginEntry.Name))
+				{
+					PluginInfo.Size = *SizePtr;
+				}
+
+				///
+				/// validation
+				///
+				if (const int64* SizePtr = PluginToSizeMap.Find(PluginEntry.Name))
+				{
+					if (CookMetadata.GetSizesPresent() != UE::Cook::ECookMetadataSizesPresent::NotPresent)
 					{
-						UE_LOG(LogInsights, Warning, TEXT("Plugin %s found with mismatched ucookmetadata and internal asset size calculation. Metadata size: %lld Calculated size: %lld"), 
-							StoredPluginName, PluginInfo.Size, *SizePtr);
+						int64 TotalSizeOfPluginInMetadata = PluginEntry.ExclusiveSizes[UE::Cook::EPluginSizeTypes::Installed];
+						if (*SizePtr != TotalSizeOfPluginInMetadata)
+						{
+							UE_LOG(LogInsights, Warning, TEXT("Plugin %s found with mismatched ucookmetadata and internal asset size calculation. Metadata size: %lld Calculated size: %lld"),
+								StoredPluginName, PluginInfo.Size, *SizePtr);
+						}
+						TotalMismatch += FMath::Abs(*SizePtr - TotalSizeOfPluginInMetadata);
+						TotalSize += TotalSizeOfPluginInMetadata;
 					}
-					TotalMismatch += FMath::Abs(*SizePtr - TotalSizeOfPluginInMetadata);
-					TotalSize += TotalSizeOfPluginInMetadata;
 				}
-			}
 
-			const int32 PluginIndex = AssetTable->GetIndexForPlugin(StoredPluginName);
-			for (uint32 DependencyIndexInMetadata = PluginEntry.DependencyIndexStart; DependencyIndexInMetadata < PluginEntry.DependencyIndexEnd; DependencyIndexInMetadata++)
-			{
-				const UE::Cook::FCookMetadataPluginEntry& DependentPluginEntry = PluginHierarchy.PluginsEnabledAtCook[PluginHierarchy.PluginDependencies[DependencyIndexInMetadata]];
-				const TCHAR* StoredReferencePluginName = AssetTable->StoreStr(DependentPluginEntry.Name);
-				int32 DependencyIndex = AssetTable->GetIndexForPlugin(StoredReferencePluginName);
-				if (DependencyIndex == -1)
+				const int32 PluginIndex = AssetTable->GetIndexForPlugin(StoredPluginName);
+				for (uint32 DependencyIndexInMetadata = PluginEntry.DependencyIndexStart; DependencyIndexInMetadata < PluginEntry.DependencyIndexEnd; DependencyIndexInMetadata++)
 				{
-					DependencyIndex = AssetTable->GetNumPlugins();
-					AssetTable->GetOrCreatePluginInfo(StoredReferencePluginName);
+					const UE::Cook::FCookMetadataPluginEntry& DependentPluginEntry = PluginHierarchy.PluginsEnabledAtCook[PluginHierarchy.PluginDependencies[DependencyIndexInMetadata]];
+					const TCHAR* StoredReferencePluginName = AssetTable->StoreStr(DependentPluginEntry.Name);
+					int32 DependencyIndex = AssetTable->GetIndexForPlugin(StoredReferencePluginName);
+					if (DependencyIndex == -1)
+					{
+						DependencyIndex = AssetTable->GetNumPlugins();
+						AssetTable->GetOrCreatePluginInfo(StoredReferencePluginName);
+					}
+					// Note that we can't use PluginInfo directly because the above code could have grown the array 
+					// and that would invalidate the reference
+					AssetTable->GetOrCreatePluginInfo(StoredPluginName).PluginDependencies.AddUnique(DependencyIndex);
+					AssetTable->GetPluginInfoByIndex(DependencyIndex).PluginReferencers.AddUnique(AssetTable->GetIndexForPlugin(StoredPluginName));
 				}
-				// Note that we can't use PluginInfo directly because the above code could have grown the array 
-				// and that would invalidate the reference
-				AssetTable->GetOrCreatePluginInfo(StoredPluginName).PluginDependencies.AddUnique(DependencyIndex);
-				AssetTable->GetPluginInfoByIndex(DependencyIndex).PluginReferencers.AddUnique(AssetTable->GetIndexForPlugin(StoredPluginName));
 			}
-		}
 
-		if (TotalMismatch != 0)
-		{
-			UE_LOG(LogInsights, Warning, TEXT("Total Size of Plugins from Metadata: %lld // Total Discrepancy (vs calculated size from assets): %lld"), TotalSize, TotalMismatch);
-		}
+			if (TotalMismatch != 0)
+			{
+				UE_LOG(LogInsights, Warning, TEXT("Total Size of Plugins from Metadata: %lld // Total Discrepancy (vs calculated size from assets): %lld"), TotalSize, TotalMismatch);
+			}
 
-		for (uint16 RootPluginIndex : PluginHierarchy.RootPlugins)
-		{
-			const UE::Cook::FCookMetadataPluginEntry& PluginEntry = PluginHierarchy.PluginsEnabledAtCook[RootPluginIndex];
-			const TCHAR* StoredPluginName = AssetTable->StoreStr(PluginEntry.Name);
-			FAssetTablePluginInfo& PluginInfo = AssetTable->GetOrCreatePluginInfo(StoredPluginName);
-			UE_LOG(LogInsights, Display, TEXT("Found root plugin %s"), StoredPluginName);
-			PluginInfo.bIsRootPlugin = true;
+			for (uint16 RootPluginIndex : PluginHierarchy.RootPlugins)
+			{
+				const UE::Cook::FCookMetadataPluginEntry& PluginEntry = PluginHierarchy.PluginsEnabledAtCook[RootPluginIndex];
+				const TCHAR* StoredPluginName = AssetTable->StoreStr(PluginEntry.Name);
+				FAssetTablePluginInfo& PluginInfo = AssetTable->GetOrCreatePluginInfo(StoredPluginName);
+				UE_LOG(LogInsights, Display, TEXT("Found root plugin %s"), StoredPluginName);
+				PluginInfo.bIsRootPlugin = true;
+			}
 		}
 	}
 	else
@@ -2000,12 +2066,20 @@ void SAssetTableTreeView::RefreshAssets()
 		}
 	}
 
+	for (int32 PluginIndex = 0; PluginIndex < AssetTable->GetNumPlugins(); PluginIndex++)
+	{
+		const FAssetTablePluginInfo& PluginInfo = AssetTable->GetPluginInfoByIndex(PluginIndex);
+		if (!PluginInfo.IsRootPlugin() && PluginInfo.GetNumReferencers() == 0 && PluginInfo.GetSize() > 0)
+		{
+			UE_LOG(LogInsights, Warning, TEXT("Found orphaned plugin %s"), PluginInfo.GetName());
+		}
+	}
 
 	// Disabling this for now as there are too many differences at the moment. Keeping the code because eventually
 	// we probably do want to use it to discover spurious dependencies or dependencies on code+content where content should
 	// get separated from code
 	// 
-	//DumpDifferencesBetweenDiscoveredDataAndLoadedMetadata(DiscoveredPluginDependencyEdges, AssetTable);
+	// DumpDifferencesBetweenDiscoveredDataAndLoadedMetadata(DiscoveredPluginDependencyEdges);
 
 
 	Stopwatch.Stop();
