Reference Viewer: added option to collapse external referencers into their corresponding owning asset.
Can be toggled via show menu or with the 'X' key.

diff --git a/Engine/Plugins/Editor/AssetManagerEditor/Source/AssetManagerEditor/Private/AssetManagerEditorCommands.cpp b/Engine/Plugins/Editor/AssetManagerEditor/Source/AssetManagerEditor/Private/AssetManagerEditorCommands.cpp
--- a/Engine/Plugins/Editor/AssetManagerEditor/Source/AssetManagerEditor/Private/AssetManagerEditorCommands.cpp
+++ b/Engine/Plugins/Editor/AssetManagerEditor/Source/AssetManagerEditor/Private/AssetManagerEditorCommands.cpp
@@ -51,6 +51,7 @@ void FAssetManagerEditorCommands::RegisterCommands()
 
 	UI_COMMAND(ShowDuplicates, "Show Duplicate References", "Toggles visibility of Duplicate References", EUserInterfaceActionType::ToggleButton, FInputChord(EKeys::U));
 	UI_COMMAND(CompactMode, "Compact Mode", "Toggles Compact View", EUserInterfaceActionType::ToggleButton, FInputChord(EKeys::V));
+	UI_COMMAND(ShowExternalReferencers, "Show External Referencers", "Toggles Show External Referencers", EUserInterfaceActionType::ToggleButton, FInputChord(EKeys::X));
 	UI_COMMAND(FilterSearch, "Filter Search Results", "Toggles filtering of search results", EUserInterfaceActionType::ToggleButton, FInputChord());
 
 	UI_COMMAND(ShowCommentPath, "Show Asset Path", "Toggles visibility of the Asset Path shown as a comment", EUserInterfaceActionType::ToggleButton, FInputChord(EKeys::P));
diff --git a/Engine/Plugins/Editor/AssetManagerEditor/Source/AssetManagerEditor/Private/AssetManagerEditorCommands.h b/Engine/Plugins/Editor/AssetManagerEditor/Source/AssetManagerEditor/Private/AssetManagerEditorCommands.h
--- a/Engine/Plugins/Editor/AssetManagerEditor/Source/AssetManagerEditor/Private/AssetManagerEditorCommands.h
+++ b/Engine/Plugins/Editor/AssetManagerEditor/Source/AssetManagerEditor/Private/AssetManagerEditorCommands.h
@@ -89,6 +89,9 @@ public:
 	// Toggles Compact Mode
 	TSharedPtr<FUICommandInfo> CompactMode;
 
+	// Toggles Show External Referencers
+	TSharedPtr<FUICommandInfo> ShowExternalReferencers;
+
 	// Toggles Path Comment 
 	TSharedPtr<FUICommandInfo> ShowCommentPath;
 
diff --git a/Engine/Plugins/Editor/AssetManagerEditor/Source/AssetManagerEditor/Private/ReferenceViewer/EdGraph_ReferenceViewer.cpp b/Engine/Plugins/Editor/AssetManagerEditor/Source/AssetManagerEditor/Private/ReferenceViewer/EdGraph_ReferenceViewer.cpp
--- a/Engine/Plugins/Editor/AssetManagerEditor/Source/AssetManagerEditor/Private/ReferenceViewer/EdGraph_ReferenceViewer.cpp
+++ b/Engine/Plugins/Editor/AssetManagerEditor/Source/AssetManagerEditor/Private/ReferenceViewer/EdGraph_ReferenceViewer.cpp
@@ -538,6 +538,45 @@ void UEdGraph_ReferenceViewer::GetSortedLinks(const TArray<FAssetIdentifier>& Id
 		if (bReferencers)
 		{
 			AssetRegistry.GetReferencers(AssetId, LinksToAsset, Query.Categories, Query.Flags);
+
+			if (!Settings->IsShowExternalReferencers())
+			{
+				TSet<FName> PackageNames;
+				for (const FAssetDependency& LinkToAsset : LinksToAsset)
+				{
+					if (!LinkToAsset.AssetId.IsValue() && !LinkToAsset.AssetId.PackageName.IsNone())
+					{
+						PackageNames.Add(LinkToAsset.AssetId.PackageName);
+					}
+				}
+
+				TMap<FName, FAssetData> PackagesToAssetDataMap;
+				UE::AssetRegistry::GetAssetForPackages(PackageNames.Array(), PackagesToAssetDataMap);
+
+				TSet<FName> OuterPathNames;
+				for (int32 LinksToAssetIndex = 0; LinksToAssetIndex < LinksToAsset.Num(); LinksToAssetIndex++)
+				{
+					FAssetDependency& AssetDependency = LinksToAsset[LinksToAssetIndex];				
+					if (FAssetData* AssetData = PackagesToAssetDataMap.Find(AssetDependency.AssetId.PackageName))
+					{
+						if (FName OuterPathName = AssetData->GetOptionalOuterPathName(); !OuterPathName.IsNone())
+						{
+							if (!OuterPathNames.Contains(OuterPathName))
+							{
+								FAssetDependency OuterDependency;
+								OuterDependency.AssetId = FAssetIdentifier(*FSoftObjectPath(OuterPathName.ToString()).GetLongPackageName());
+								OuterDependency.Category = AssetDependency.Category;
+								OuterDependency.Properties = AssetDependency.Properties;
+								LinksToAsset.Add(OuterDependency);
+
+								OuterPathNames.Add(OuterPathName);
+							}
+
+							LinksToAsset.RemoveAtSwap(LinksToAssetIndex--);
+						}
+					}
+				}
+			}
 		}
 		else
 		{
diff --git a/Engine/Plugins/Editor/AssetManagerEditor/Source/AssetManagerEditor/Private/ReferenceViewer/ReferenceViewerSettings.cpp b/Engine/Plugins/Editor/AssetManagerEditor/Source/AssetManagerEditor/Private/ReferenceViewer/ReferenceViewerSettings.cpp
--- a/Engine/Plugins/Editor/AssetManagerEditor/Source/AssetManagerEditor/Private/ReferenceViewer/ReferenceViewerSettings.cpp
+++ b/Engine/Plugins/Editor/AssetManagerEditor/Source/AssetManagerEditor/Private/ReferenceViewer/ReferenceViewerSettings.cpp
@@ -34,6 +34,11 @@ bool UReferenceViewerSettings::IsCompactMode() const
 	return bIsCompactMode;
 }
 
+bool UReferenceViewerSettings::IsShowExternalReferencers() const
+{
+	return bIsShowExternalReferencers;
+}
+
 bool UReferenceViewerSettings::IsShowDuplicates() const
 {
 	return bIsShowDuplicates;
@@ -105,6 +110,12 @@ void UReferenceViewerSettings::SetCompactModeEnabled(bool bNewEnabled)
 	SaveConfig();
 }
 
+void UReferenceViewerSettings::SetShowExternalReferencersEnabled(bool bEnabled)
+{
+	bIsShowExternalReferencers = bEnabled;
+	SaveConfig();
+}
+
 void UReferenceViewerSettings::SetShowDuplicatesEnabled(bool bNewEnabled)
 {
 	bIsShowDuplicates = bNewEnabled;
diff --git a/Engine/Plugins/Editor/AssetManagerEditor/Source/AssetManagerEditor/Private/ReferenceViewer/SReferenceViewer.cpp b/Engine/Plugins/Editor/AssetManagerEditor/Source/AssetManagerEditor/Private/ReferenceViewer/SReferenceViewer.cpp
--- a/Engine/Plugins/Editor/AssetManagerEditor/Source/AssetManagerEditor/Private/ReferenceViewer/SReferenceViewer.cpp
+++ b/Engine/Plugins/Editor/AssetManagerEditor/Source/AssetManagerEditor/Private/ReferenceViewer/SReferenceViewer.cpp
@@ -1338,6 +1338,17 @@ bool SReferenceViewer::IsCompactModeChecked() const
 	return Settings->IsCompactMode();
 }
 
+void SReferenceViewer::OnShowExternalReferencersChanged()
+{
+	Settings->SetShowExternalReferencersEnabled(!Settings->IsShowExternalReferencers());
+	RebuildGraph();
+}
+
+bool SReferenceViewer::IsShowExternalReferencersChecked() const
+{
+	return Settings->IsShowExternalReferencers();
+}
+
 void SReferenceViewer::OnShowDuplicatesChanged()
 {
 	Settings->SetShowDuplicatesEnabled(!Settings->IsShowDuplicates());
@@ -1551,6 +1562,12 @@ void SReferenceViewer::RegisterActions()
 		FIsActionChecked::CreateSP(this, &SReferenceViewer::IsCompactModeChecked),
 		FIsActionButtonVisible::CreateLambda([this] { return bShowCompactMode; }));
 
+	ReferenceViewerActions->MapAction(
+		FAssetManagerEditorCommands::Get().ShowExternalReferencers,
+		FExecuteAction::CreateSP(this, &SReferenceViewer::OnShowExternalReferencersChanged),
+		FCanExecuteAction(),	
+		FIsActionChecked::CreateSP(this, &SReferenceViewer::IsShowExternalReferencersChecked));
+
 	ReferenceViewerActions->MapAction(
 		FAssetManagerEditorCommands::Get().FilterSearch,
 		FExecuteAction::CreateSP(this, &SReferenceViewer::OnShowFilteredPackagesOnlyChanged),
@@ -2254,6 +2271,7 @@ TSharedRef<SWidget> SReferenceViewer::GetShowMenuContent()
 	MenuBuilder.AddMenuEntry(FAssetManagerEditorCommands::Get().ShowDuplicates);
 	MenuBuilder.AddMenuEntry(FAssetManagerEditorCommands::Get().FilterSearch);
 	MenuBuilder.AddMenuEntry(FAssetManagerEditorCommands::Get().CompactMode);
+	MenuBuilder.AddMenuEntry(FAssetManagerEditorCommands::Get().ShowExternalReferencers);
 	MenuBuilder.AddMenuEntry(FAssetManagerEditorCommands::Get().ShowCommentPath);
 	MenuBuilder.EndSection();
 
diff --git a/Engine/Plugins/Editor/AssetManagerEditor/Source/AssetManagerEditor/Private/ReferenceViewer/SReferenceViewer.h b/Engine/Plugins/Editor/AssetManagerEditor/Source/AssetManagerEditor/Private/ReferenceViewer/SReferenceViewer.h
--- a/Engine/Plugins/Editor/AssetManagerEditor/Source/AssetManagerEditor/Private/ReferenceViewer/SReferenceViewer.h
+++ b/Engine/Plugins/Editor/AssetManagerEditor/Source/AssetManagerEditor/Private/ReferenceViewer/SReferenceViewer.h
@@ -147,6 +147,9 @@ private:
 	void OnCompactModeChanged();
 	bool IsCompactModeChecked() const;
 
+	void OnShowExternalReferencersChanged();
+	bool IsShowExternalReferencersChecked() const;
+
 	void OnShowDuplicatesChanged();
 	bool IsShowDuplicatesChecked() const;
 
diff --git a/Engine/Plugins/Editor/AssetManagerEditor/Source/AssetManagerEditor/Public/ReferenceViewer/ReferenceViewerSettings.h b/Engine/Plugins/Editor/AssetManagerEditor/Source/AssetManagerEditor/Public/ReferenceViewer/ReferenceViewerSettings.h
--- a/Engine/Plugins/Editor/AssetManagerEditor/Source/AssetManagerEditor/Public/ReferenceViewer/ReferenceViewerSettings.h
+++ b/Engine/Plugins/Editor/AssetManagerEditor/Source/AssetManagerEditor/Public/ReferenceViewer/ReferenceViewerSettings.h
@@ -92,6 +92,9 @@ public:
 	bool IsCompactMode() const;
 	void SetCompactModeEnabled(bool bNewEnabled);
 
+	bool IsShowExternalReferencers() const;
+	void SetShowExternalReferencersEnabled(bool bNewEnabled);
+
 	bool IsShowPath() const;
 	void SetShowPathEnabled(bool bNewEnabled);
 
@@ -180,6 +183,10 @@ private:
 	UPROPERTY(config)
 	bool bIsCompactMode;
 
+	/* Whether to show external referencers or collapsed owner asset */
+	UPROPERTY(config)
+	bool bIsShowExternalReferencers;
+
 	/* Whether to show the package's path as a comment */
 	UPROPERTY(config)
 	bool bIsShowPath;
