Asset Disk Size tool now consumes custom column data from ucookmeta. Added a column for plugin type as well (e.g., shader pseudoplugin, engine pseudoplugin, etc)

diff --git a/Engine/Plugins/Editor/AssetManagerEditor/Source/AssetManagerEditor/Private/TreeView/AssetTable.cpp b/Engine/Plugins/Editor/AssetManagerEditor/Source/AssetManagerEditor/Private/TreeView/AssetTable.cpp
--- a/Engine/Plugins/Editor/AssetManagerEditor/Source/AssetManagerEditor/Private/TreeView/AssetTable.cpp
+++ b/Engine/Plugins/Editor/AssetManagerEditor/Source/AssetManagerEditor/Private/TreeView/AssetTable.cpp
@@ -36,6 +36,7 @@ const FName FAssetTableColumns::ChunksColumnId(TEXT("Chunks"));
 const FName FAssetTableColumns::NativeClassColumnId(TEXT("NativeClass"));
 const FName FAssetTableColumns::PluginNameColumnId(TEXT("PluginName"));
 const FName FAssetTableColumns::PluginInclusiveSizeColumnId(TEXT("PluginInclusiveSize"));
+const FName FAssetTableColumns::PluginTypeColumnId(TEXT("PluginType"));
 
 ////////////////////////////////////////////////////////////////////////////////////////////////////
 // FAssetTableStringValueGetterWithDependencyAggregationHandling
@@ -1341,7 +1342,168 @@ void FAssetTable::AddDefaultColumns()
 
 		AddColumn(ColumnRef);
 	}
-	
+
+	//////////////////////////////////////////////////
+	// PluginType Column
+	{
+		TSharedRef<FTableColumn> ColumnRef = MakeShared<FTableColumn>(FAssetTableColumns::PluginTypeColumnId);
+		FTableColumn& Column = *ColumnRef;
+
+		Column.SetIndex(ColumnIndex++);
+
+		Column.SetShortName(LOCTEXT("PluginTypeColumnName", "Plugin Type"));
+		Column.SetTitleName(LOCTEXT("PluginTypeColumnTitle", "Plugin Type"));
+		Column.SetDescription(LOCTEXT("PluginTypeColumnDesc", "Plugin Type (e.g., Normal, Root, Shader Pseudoplugin, etc.)"));
+
+		Column.SetFlags(ETableColumnFlags::CanBeHidden | ETableColumnFlags::CanBeFiltered);
+
+		Column.SetHorizontalAlignment(HAlign_Left);
+		Column.SetInitialWidth(75.0f);
+
+		Column.SetDataType(ETableCellDataType::CString);
+
+		class FPluginTypeValueGetter : public FTableCellValueGetter
+		{
+		public:
+			virtual const TOptional<FTableCellValue> GetValue(const FTableColumn& Column, const FBaseTreeNode& Node) const override
+			{
+				if (Node.Is<FPluginSimpleGroupNode>())
+				{
+					TSharedPtr<FTable> TablePtr = static_cast<const FPluginSimpleGroupNode&>(Node).GetParentTable().Pin();
+					const FAssetTable& AssetTable = static_cast<const FAssetTable&>(*TablePtr);
+					int32 PluginIndex = static_cast<const FPluginSimpleGroupNode&>(Node).GetPluginIndex();
+					return FTableCellValue(AssetTable.GetPluginInfoByIndex(PluginIndex).GetPluginType());
+				}
+
+				return TOptional<FTableCellValue>();
+			}
+		};
+		TSharedRef<ITableCellValueGetter> Getter = MakeShared<FPluginTypeValueGetter>();
+		Column.SetValueGetter(Getter);
+
+		TSharedRef<ITableCellValueFormatter> Formatter = MakeShared<FCStringValueFormatterAsText>();
+		Column.SetValueFormatter(Formatter);
+
+		TSharedRef<ITableCellValueSorter> Sorter = MakeShared<FSorterByCStringValue>(ColumnRef);
+		Column.SetValueSorter(Sorter);
+		Column.SetInitialSortMode(EColumnSortMode::Descending);
+
+		Column.SetAggregation(ETableColumnAggregation::SameValue);
+
+		AddColumn(ColumnRef);
+	}
+
+	//////////////////////////////////////////////////
+	// Custom Columns
+	for (int32 CustomColumnIndex = 0; CustomColumnIndex < CustomColumns.Num(); CustomColumnIndex++)
+	{
+		const FCustomColumnDefinition& CustomColumn = CustomColumns[CustomColumnIndex];
+
+		TSharedRef<FTableColumn> ColumnRef = MakeShared<FTableColumn>(CustomColumn.ColumnId);
+		FTableColumn& Column = *ColumnRef;
+
+		Column.SetIndex(ColumnIndex++);
+
+		const FText NameAsText = FText::FromName(CustomColumn.ColumnId);
+		Column.SetShortName(NameAsText);
+		Column.SetTitleName(NameAsText);
+		Column.SetDescription(LOCTEXT("CustomColumnDesc", "Custom column data"));
+
+		Column.SetFlags(ETableColumnFlags::CanBeHidden | ETableColumnFlags::CanBeFiltered);
+
+		Column.SetHorizontalAlignment(HAlign_Left);
+		if (CustomColumn.Type == ECustomColumnDefinitionType::Boolean)
+		{
+			Column.SetInitialWidth(50.f);
+			Column.SetDataType(ETableCellDataType::Bool);
+			class FCustomColumnBoolValueGetter : public FTableCellValueGetter
+			{
+			public:
+
+				FCustomColumnBoolValueGetter(int32 InIndexWithinRowData) :
+					FTableCellValueGetter(),
+					IndexWithinRowData(InIndexWithinRowData) {}
+
+				virtual const TOptional<FTableCellValue> GetValue(const FTableColumn& Column, const FBaseTreeNode& Node) const override
+				{
+					if (Node.Is<FPluginSimpleGroupNode>() && !Node.Is<FPluginDependenciesGroupNode>())
+					{
+						// This node represents a single plugin (it might be the plugin itself or the plugin+deps node for that plugin)
+						const FPluginSimpleGroupNode& PluginNode = Node.As<FPluginSimpleGroupNode>();
+						TSharedPtr<FTable> TablePtr = PluginNode.GetParentTable().Pin();
+						const FAssetTable& AssetTable = static_cast<const FAssetTable&>(*TablePtr);
+						const FAssetTablePluginInfo& PluginInfo = AssetTable.GetPluginInfoByIndexChecked(PluginNode.GetPluginIndex());
+						if (const bool* Value = PluginInfo.TryGetDataByKey<bool>(IndexWithinRowData))
+						{
+							return FTableCellValue(*Value);
+						}
+					}
+
+					return TOptional<FTableCellValue>();
+				}
+
+			private:
+				int32 IndexWithinRowData;
+			};
+
+			TSharedRef<ITableCellValueGetter> Getter = MakeShared<FCustomColumnBoolValueGetter>(CustomColumnIndex);
+			Column.SetValueGetter(Getter);
+
+			TSharedRef<ITableCellValueFormatter> Formatter = MakeShared<FBoolValueFormatterAsTrueFalse>();
+			Column.SetValueFormatter(Formatter);
+
+			TSharedRef<ITableCellValueSorter> Sorter = MakeShared<FSorterByBoolValue>(ColumnRef);
+			Column.SetValueSorter(Sorter);
+		}
+		else
+		{
+			Column.SetInitialWidth(200.f);
+			Column.SetDataType(ETableCellDataType::CString);
+			class FCustomColumnStringValueGetter : public FTableCellValueGetter
+			{
+			public:
+
+				FCustomColumnStringValueGetter(int32 InIndexWithinRowData) :
+					FTableCellValueGetter(),
+					IndexWithinRowData(InIndexWithinRowData) {}
+
+				virtual const TOptional<FTableCellValue> GetValue(const FTableColumn& Column, const FBaseTreeNode& Node) const override
+				{
+					if (Node.Is<FPluginSimpleGroupNode>() && !Node.Is<FPluginDependenciesGroupNode>())
+					{
+						// This node represents a single plugin (it might be the plugin itself or the plugin+deps node for that plugin)
+						const FPluginSimpleGroupNode& PluginNode = Node.As<FPluginSimpleGroupNode>();
+						TSharedPtr<FTable> TablePtr = PluginNode.GetParentTable().Pin();
+						const FAssetTable& AssetTable = static_cast<const FAssetTable&>(*TablePtr);
+						const FAssetTablePluginInfo& PluginInfo = AssetTable.GetPluginInfoByIndexChecked(PluginNode.GetPluginIndex());
+						if (const TCHAR* const* Value = PluginInfo.TryGetDataByKey<const TCHAR*>(IndexWithinRowData))
+						{
+							return FTableCellValue(*Value);
+						}
+					}
+
+					return TOptional<FTableCellValue>();
+				}
+
+			private:
+				int32 IndexWithinRowData;
+			};
+
+			TSharedRef<ITableCellValueGetter> Getter = MakeShared<FCustomColumnStringValueGetter>(CustomColumnIndex);
+			Column.SetValueGetter(Getter);
+
+			TSharedRef<ITableCellValueFormatter> Formatter = MakeShared<FCStringValueFormatterAsText>();
+			Column.SetValueFormatter(Formatter);
+
+			TSharedRef<ITableCellValueSorter> Sorter = MakeShared<FSorterByCStringValue>(ColumnRef);
+			Column.SetValueSorter(Sorter);
+		}
+
+
+		Column.SetAggregation(ETableColumnAggregation::SameValue);
+
+		AddColumn(ColumnRef);
+	}
 }
 
 /*static*/TSet<int32> FAssetTableRow::GatherAllReachableNodes(const TArray<int32>& StartingNodes, const FAssetTable& OwningTable, const TSet<int32>& AdditionalNodesToStopAt, const TSet<const TCHAR*, TStringPointerSetKeyFuncs_DEPRECATED<const TCHAR*>>& RestrictToPlugins, TMap<int32, TArray<int32>>* OutRouteMap /*= nullptr*/)
diff --git a/Engine/Plugins/Editor/AssetManagerEditor/Source/AssetManagerEditor/Private/TreeView/AssetTable.h b/Engine/Plugins/Editor/AssetManagerEditor/Source/AssetManagerEditor/Private/TreeView/AssetTable.h
--- a/Engine/Plugins/Editor/AssetManagerEditor/Source/AssetManagerEditor/Private/TreeView/AssetTable.h
+++ b/Engine/Plugins/Editor/AssetManagerEditor/Source/AssetManagerEditor/Private/TreeView/AssetTable.h
@@ -8,6 +8,7 @@
 #include "Containers/StringView.h"
 #include "Insights/Table/ViewModels/Table.h"
 #include "Misc/CString.h"
+#include "Misc/TVariant.h"
 
 namespace UE
 {
@@ -126,6 +127,7 @@ struct FAssetTableColumns
 	static const FName NativeClassColumnId;
 	static const FName PluginNameColumnId;
 	static const FName PluginInclusiveSizeColumnId;
+	static const FName PluginTypeColumnId;
 };
 
 struct FAssetTableDependencySizes
@@ -148,6 +150,7 @@ public:
 	const TArray<int32>& GetReferencers() const { return PluginReferencers; }
 	int64 GetSize() const { return Size; }
 	bool IsRootPlugin() const { return bIsRootPlugin; }
+	const TCHAR* GetPluginType() const { return PluginTypeString; }
 
 	int64 GetOrComputeTotalSizeInclusiveOfDependencies(const FAssetTable& OwningTable) const;
 
@@ -158,12 +161,43 @@ public:
 
 	static void ComputeTotalSelfAndInclusiveSizes(const FAssetTable& OwningTable, const TSet<int32>& RootPlugins, int64& OutTotalSelfSize, int64& OutTotalInclusiveSize);
 
+	template<typename T> const T* TryGetDataByKey(int32 Key) const
+	{
+		// In general there are expected to be very few keys since each key corresponds to a custom implemented
+		// column. Therefore a simple linear search is expected to be sufficiently performant.
+		for (const FCustomColumnData& CustomDataEntry : CustomColumnData)
+		{
+			if (CustomDataEntry.Key > Key)
+			{
+				// We've passed where this key ought to be found. No entry is present.
+				break; 
+			}
+			else if (CustomDataEntry.Key == Key)
+			{
+				return CustomDataEntry.Value.TryGet<T>();
+			}
+		}
+		return nullptr;
+	}
+
 private:
 	void ComputeDependencySizes(const FAssetTable& OwningTable) const;
 
 	TArray<int32> PluginDependencies;
 	TArray<int32> PluginReferencers;
+
+	typedef TVariant<bool, const TCHAR*> CustomFieldVariantType;
+
+	struct FCustomColumnData
+	{
+		CustomFieldVariantType Value;
+		int32 Key;
+	};
+
+	TArray<FCustomColumnData> CustomColumnData;
+
 	const TCHAR* PluginName = nullptr;
+	const TCHAR* PluginTypeString = nullptr;
 	int64 Size = 0;
 	mutable int64 InclusiveSize = -1;
 	mutable int64 UniqueDependenciesSize = -1;
@@ -298,8 +332,24 @@ public:
 		StringStore.Reset();
 		Plugins.Empty();
 		PluginNameToIndexMap.Empty();
+		CustomColumns.Empty();
 	}
 
+	enum class ECustomColumnDefinitionType
+	{
+		Boolean,
+		String
+	};
+
+	struct FCustomColumnDefinition
+	{
+		FName ColumnId;
+		ECustomColumnDefinitionType Type;
+		int32 Key;
+	};
+
+	void AddCustomColumn(const FCustomColumnDefinition& ColumnDefinition) { CustomColumns.Add(ColumnDefinition); }
+
 private:
 	void AddDefaultColumns();
 
@@ -307,6 +357,7 @@ private:
 	TArray<FAssetTableRow> Assets;
 	TArray<FAssetTablePluginInfo> Plugins;
 	TMap<const TCHAR*, int32, FDefaultSetAllocator, TStringPointerMapKeyFuncs_DEPRECATED<const TCHAR*, int32>> PluginNameToIndexMap;
+	TArray<FCustomColumnDefinition> CustomColumns;
 	int32 VisibleAssetCount = 0;
 	FAssetTableStringStore StringStore;
 };
diff --git a/Engine/Plugins/Editor/AssetManagerEditor/Source/AssetManagerEditor/Private/TreeView/SAssetTableTreeView.cpp b/Engine/Plugins/Editor/AssetManagerEditor/Source/AssetManagerEditor/Private/TreeView/SAssetTableTreeView.cpp
--- a/Engine/Plugins/Editor/AssetManagerEditor/Source/AssetManagerEditor/Private/TreeView/SAssetTableTreeView.cpp
+++ b/Engine/Plugins/Editor/AssetManagerEditor/Source/AssetManagerEditor/Private/TreeView/SAssetTableTreeView.cpp
@@ -109,6 +109,13 @@ void SAssetTableTreeView::Tick(const FGeometry& AllottedGeometry, const double I
 		{
 			bNeedsToRebuild = false;
 			RebuildTree(true);
+			if (bNeedsToRebuildColumns)
+			{
+				// This resets the column information only. It does NOT nuke the table data.
+				GetAssetTable()->Reset();
+				RebuildColumns();
+				bNeedsToRebuildColumns = false;
+			}
 			ApplyViewPreset(*SelectedViewPreset);
 		}
 	}
@@ -1964,6 +1971,27 @@ void SAssetTableTreeView::RefreshAssets()
 		// Plugin Data
 		{
 			const UE::Cook::FCookMetadataPluginHierarchy& PluginHierarchy = CookMetadata.GetPluginHierarchy();
+
+			for (int32 CustomColumnIndex = 0; CustomColumnIndex < PluginHierarchy.CustomFieldEntries.Num(); CustomColumnIndex++)
+			{
+				FAssetTable::FCustomColumnDefinition CustomColumnDefinition;
+				CustomColumnDefinition.ColumnId = FName(PluginHierarchy.CustomFieldEntries[CustomColumnIndex].Name);
+				UE::Cook::ECookMetadataCustomFieldType FieldType = PluginHierarchy.CustomFieldEntries[CustomColumnIndex].Type;
+				if (ensureMsgf(FieldType != UE::Cook::ECookMetadataCustomFieldType::Unknown, TEXT("Unknown field type encountered")))
+				{
+					if (FieldType == UE::Cook::ECookMetadataCustomFieldType::String)
+					{
+						CustomColumnDefinition.Type = FAssetTable::ECustomColumnDefinitionType::String;
+					}
+					else if (FieldType == UE::Cook::ECookMetadataCustomFieldType::Bool)
+					{
+						CustomColumnDefinition.Type = FAssetTable::ECustomColumnDefinitionType::Boolean;
+					}
+				}
+				CustomColumnDefinition.Key = CustomColumnIndex;
+				AssetTable->AddCustomColumn(CustomColumnDefinition);
+			}
+
 			int64 TotalMismatch = 0;
 			int64 TotalSize = 0;
 			for (const UE::Cook::FCookMetadataPluginEntry& PluginEntry : PluginHierarchy.PluginsEnabledAtCook)
@@ -1979,6 +2007,28 @@ void SAssetTableTreeView::RefreshAssets()
 					PluginInfo.Size = *SizePtr;
 				}
 
+				PluginInfo.PluginTypeString = AssetTable->StoreStr(*PluginEntry.GetPluginTypeAsText().ToString());
+
+				for (int32 CustomColumnIndex = 0; CustomColumnIndex < PluginHierarchy.CustomFieldEntries.Num(); CustomColumnIndex++)
+				{
+					if (const UE::Cook::FCookMetadataPluginEntry::CustomFieldVariantType* VariantEntry = PluginEntry.CustomFields.Find(CustomColumnIndex))
+					{
+						FAssetTablePluginInfo::FCustomColumnData ColumnData;
+						ColumnData.Key = CustomColumnIndex;
+
+						if (PluginHierarchy.CustomFieldEntries[CustomColumnIndex].Type == UE::Cook::ECookMetadataCustomFieldType::Bool)
+						{
+							ColumnData.Value.Set<bool>(VariantEntry->Get<bool>());
+						}
+						else if (PluginHierarchy.CustomFieldEntries[CustomColumnIndex].Type == UE::Cook::ECookMetadataCustomFieldType::String)
+						{
+							ColumnData.Value.Set<const TCHAR*>(AssetTable->StoreStr(*VariantEntry->Get<FString>()));
+						}
+
+						PluginInfo.CustomColumnData.Add(ColumnData);
+					}
+				}
+
 				///
 				/// validation
 				///
@@ -2093,7 +2143,7 @@ void SAssetTableTreeView::RefreshAssets()
 		StringStore.GetNumStrings(), StringStore.GetTotalStringSize(), StringStore.GetAllocatedSize(),
 		(double)StringStore.GetAllocatedSize() * 100.0 / (double)StringStore.GetTotalInputStringSize());
 
-	RequestRebuildTree();
+	RequestRebuildTree(/*NeedsColumnRebuild =*/true);
 }
 
 void SAssetTableTreeView::ClearTableAndTree()
@@ -2217,9 +2267,10 @@ void SAssetTableTreeView::DumpDifferencesBetweenDiscoveredDataAndLoadedMetadata(
 
 ////////////////////////////////////////////////////////////////////////////////////////////////////
 
-void SAssetTableTreeView::RequestRebuildTree()
+void SAssetTableTreeView::RequestRebuildTree(bool NeedsColumnRebuild)
 {
 	bNeedsToRebuild = true;
+	bNeedsToRebuildColumns = NeedsColumnRebuild;
 	FooterLeftText = LOCTEXT("FooterLeftTextFmt_RebuildTree_Filtered", "Rebuilding tree... please wait...");
 }
 
diff --git a/Engine/Plugins/Editor/AssetManagerEditor/Source/AssetManagerEditor/Private/TreeView/SAssetTableTreeView.h b/Engine/Plugins/Editor/AssetManagerEditor/Source/AssetManagerEditor/Private/TreeView/SAssetTableTreeView.h
--- a/Engine/Plugins/Editor/AssetManagerEditor/Source/AssetManagerEditor/Private/TreeView/SAssetTableTreeView.h
+++ b/Engine/Plugins/Editor/AssetManagerEditor/Source/AssetManagerEditor/Private/TreeView/SAssetTableTreeView.h
@@ -130,7 +130,7 @@ private:
 	typedef TSet<const TCHAR*, TStringPointerSetKeyFuncs_DEPRECATED<const TCHAR*>> DeprecatedTCharSetType;
 	void DumpDifferencesBetweenDiscoveredDataAndLoadedMetadata(TMap<const TCHAR*, DeprecatedTCharSetType, FDefaultSetAllocator, TStringPointerMapKeyFuncs_DEPRECATED<const TCHAR*, DeprecatedTCharSetType>>& DiscoveredPluginDependencyEdges) const;
 
-	void RequestRebuildTree();
+	void RequestRebuildTree(bool bNeedsColumnRebuild = false);
 
 	void UpdateRegistryInfoTextPostLoad(ECheckFilesExistAndHashMatchesResult StatusResult);
 
@@ -138,6 +138,7 @@ private:
 	bool bNeedsToOpenRegistry = false;
 	bool bNeedsToRefreshAssets = false;
 	bool bNeedsToRebuild = false;
+	bool bNeedsToRebuildColumns = false;
 
 	FText FooterLeftTextStoredPreOpen;
 	FText RegistryLoadedText;
diff --git a/Engine/Source/Developer/CookMetadata/Private/CookMetadata.cpp b/Engine/Source/Developer/CookMetadata/Private/CookMetadata.cpp
--- a/Engine/Source/Developer/CookMetadata/Private/CookMetadata.cpp
+++ b/Engine/Source/Developer/CookMetadata/Private/CookMetadata.cpp
@@ -251,4 +251,20 @@ FText FCookMetadataState::GetSizesPresentAsText() const
 	return CookMetadataSizesPresentStrings[static_cast<size_t>(SizesPresent)];
 }
 
+FText FCookMetadataPluginEntry::GetPluginTypeAsText() const
+{
+	// Uncapitalized, presentation-ready
+	static FText CookMetadataPluginEntryTypeStrings[] =
+	{
+		NSLOCTEXT("CookMetadataPluginEntryType", "PluginTypeUnassigned", "unknown"),
+		NSLOCTEXT("CookMetadataPluginEntryType", "PluginTypeNormal", "normal"),
+		NSLOCTEXT("CookMetadataPluginEntryType", "PluginTypeRoot", "root"),
+		NSLOCTEXT("CookMetadataPluginEntryType", "PluginTypeEnginePseudo", "engine pseudoplugin"),
+		NSLOCTEXT("CookMetadataPluginEntryType", "PluginTypeGamePseudo", "game pseudoplugin"),
+		NSLOCTEXT("CookMetadataPluginEntryType", "PluginTypeShaderPseudo", "shader pseudoplugin")
+	};
+
+	static_assert(sizeof(CookMetadataPluginEntryTypeStrings) / sizeof(CookMetadataPluginEntryTypeStrings[0]) == static_cast<size_t>(ECookMetadataPluginType::Count));
+	return CookMetadataPluginEntryTypeStrings[static_cast<size_t>(Type)];
+}
 } // end namespace
\ No newline at end of file
diff --git a/Engine/Source/Developer/CookMetadata/Public/CookMetadata.h b/Engine/Source/Developer/CookMetadata/Public/CookMetadata.h
--- a/Engine/Source/Developer/CookMetadata/Public/CookMetadata.h
+++ b/Engine/Source/Developer/CookMetadata/Public/CookMetadata.h
@@ -144,7 +144,7 @@ enum class ECookMetadataPluginType
 	// referencing the shader, including possibly an "Unrooted" plugin.
 	ShaderPseudo,
 
-
+	Count
 };
 
 enum class ECookMetadataCustomFieldType : uint8
@@ -192,6 +192,8 @@ struct COOKMETADATA_API FCookMetadataPluginEntry
 
 	uint32 DependencyCount() const { return DependencyIndexEnd - DependencyIndexStart; }
 
+	FText GetPluginTypeAsText() const;
+
 	// !!! If you edit this, be sure to update the upgrade paths in CookMetadata.cpp
 	friend FArchive& operator<<(FArchive& Ar, FCookMetadataPluginEntry& Entry)
 	{
