[Insights] Fixed ASan "unreachable code" error.

diff --git a/Engine/Plugins/Editor/AssetManagerEditor/Source/AssetManagerEditor/Private/Insights/Table/ViewModels/TableCellValue.h b/Engine/Plugins/Editor/AssetManagerEditor/Source/AssetManagerEditor/Private/Insights/Table/ViewModels/TableCellValue.h
--- a/Engine/Plugins/Editor/AssetManagerEditor/Source/AssetManagerEditor/Private/Insights/Table/ViewModels/TableCellValue.h
+++ b/Engine/Plugins/Editor/AssetManagerEditor/Source/AssetManagerEditor/Private/Insights/Table/ViewModels/TableCellValue.h
@@ -144,7 +144,6 @@ public:
 			case ETableCellDataType::CString:	return FString(CString);
 			default:							return Custom.IsValid() ? Custom->AsText().ToString() : FString();
 		}
-		return FString();
 	}
 
 	FText AsText() const
@@ -184,7 +183,7 @@ public:
 		int64 Int64;
 		float Float;
 		double Double;
-		const TCHAR* CString; // should be valid for the lieftime of the owner table
+		const TCHAR* CString; // should be valid for the lifetime of the owner table
 		uint64 ValueId; // only used by Custom types
 	};
 
diff --git a/Engine/Source/Developer/TraceInsights/Private/Insights/Table/ViewModels/TableCellValue.h b/Engine/Source/Developer/TraceInsights/Private/Insights/Table/ViewModels/TableCellValue.h
--- a/Engine/Source/Developer/TraceInsights/Private/Insights/Table/ViewModels/TableCellValue.h
+++ b/Engine/Source/Developer/TraceInsights/Private/Insights/Table/ViewModels/TableCellValue.h
@@ -142,7 +142,6 @@ public:
 			case ETableCellDataType::CString:	return FString(CString);
 			default:							return Custom.IsValid() ? Custom->AsText().ToString() : FString();
 		}
-		return FString();
 	}
 
 	FText AsText() const
@@ -182,7 +181,7 @@ public:
 		int64 Int64;
 		float Float;
 		double Double;
-		const TCHAR* CString; // should be valid for the lieftime of the owner table
+		const TCHAR* CString; // should be valid for the lifetime of the owner table
 		uint64 ValueId; // only used by Custom types
 	};
 
