AssetManagerEditor: Add a sentry in FAssetSourceControlContextMenuState to test whether crashes in CanExecuteSCCMerge are coming from a dangling pointer.
#rnx

diff --git a/Engine/Plugins/Editor/AssetManagerEditor/Source/AssetManagerEditor/Private/AssetSourceControlContextMenu.cpp b/Engine/Plugins/Editor/AssetManagerEditor/Source/AssetManagerEditor/Private/AssetSourceControlContextMenu.cpp
--- a/Engine/Plugins/Editor/AssetManagerEditor/Source/AssetManagerEditor/Private/AssetSourceControlContextMenu.cpp
+++ b/Engine/Plugins/Editor/AssetManagerEditor/Source/AssetManagerEditor/Private/AssetSourceControlContextMenu.cpp
@@ -117,6 +117,8 @@ namespace UE::AssetSourceControlContextMenu::Private
 class FAssetSourceControlContextMenuState : public TSharedFromThis<FAssetSourceControlContextMenuState>
 {
 public:
+	~FAssetSourceControlContextMenuState();
+
 	void Initialize(FToolMenuSection& InSection);
 
 	bool IsValid() const;
@@ -194,6 +196,8 @@ private:
 	/** Cancel any currently running perforce operation */
 	void CancelCacheCanExecuteVars();
 
+	bool IsClassStillConstructed() const;
+
 private:
 
 	TArray<FAssetData> SelectedAssets;
@@ -204,6 +208,9 @@ private:
 	TSharedPtr<class ISourceControlOperation, ESPMode::ThreadSafe> SCCOperation;
 	std::atomic<EAsyncState> AsyncState = EAsyncState::None;
 
+	static constexpr int32 ValiditySentryConstant = 0x85B921C3;
+	int32 ValiditySentry = ValiditySentryConstant;
+
 	// If folders are present in the selection then we adjust some behaviors to account for the potentially large amount of files to scan
 	bool bContainsFolders = false;
 
@@ -560,9 +567,23 @@ void FAssetSourceControlContextMenuState::FillSourceControlSubMenu(UToolMenu* Me
 			FSlateIcon(FRevisionControlStyleManager::GetStyleSetName(), "RevisionControl.Actions.Merge"),
 			FUIAction(
 				FExecuteAction::CreateSP(this, &FAssetSourceControlContextMenuState::ExecuteSCCMerge),
-				FCanExecuteAction::CreateLambda([this]() { return IsActionEnabled(CanExecuteSCCMerge()); })
-			),
-			FIsAsyncProcessingActive::CreateLambda([this]() { return IsStillScanning(CanExecuteSCCMerge()); })
+				FCanExecuteAction::CreateLambda([this]() {
+					if (!IsClassStillConstructed())
+					{
+						UE_LOG(LogCore, Fatal, TEXT("Dangling pointer detected in IsActionEnabled lambda passed to FUIAction"));
+						return false;
+					}
+					return IsActionEnabled(CanExecuteSCCMerge());
+				})),
+			FIsAsyncProcessingActive::CreateLambda([this]()
+				{
+					if (!IsClassStillConstructed())
+					{
+						UE_LOG(LogCore, Fatal, TEXT("Dangling pointer detected in IsAsyncProcessingActive lambda passed to AddAsyncMenuEntry"));
+						return false;
+					}
+					return IsStillScanning(CanExecuteSCCMerge());
+				})
 		);
 	}
 
@@ -868,6 +889,19 @@ void FAssetSourceControlContextMenuState::ExecuteSCCSync() const
 	AssetViewUtils::SyncPackagesFromSourceControl(PackageNames);
 }
 
+FAssetSourceControlContextMenuState::~FAssetSourceControlContextMenuState()
+{
+	ValiditySentry = 0;
+}
+
+bool FAssetSourceControlContextMenuState::IsClassStillConstructed() const
+{
+	// UE-195836. We suspect that CanExecuteSCCMerge is being called on a dangling pointer
+	// so we have added a sentry value in the class that gets overwritten with a different 
+	// value in the destructor.
+	return ValiditySentry == ValiditySentryConstant;
+}
+
 bool FAssetSourceControlContextMenuState::CanExecuteSCCMerge() const
 {
 	FAssetToolsModule& AssetToolsModule = FModuleManager::GetModuleChecked<FAssetToolsModule>("AssetTools");
