Add a way to add all primary assets of a particular type to the Asset Audit tool
Allows for easier viewing and confirmation in the Asset Audit tool for all dependencies of a managed references

diff --git a/Engine/Plugins/Editor/AssetManagerEditor/Source/AssetManagerEditor/Private/SAssetAuditBrowser.cpp b/Engine/Plugins/Editor/AssetManagerEditor/Source/AssetManagerEditor/Private/SAssetAuditBrowser.cpp
--- a/Engine/Plugins/Editor/AssetManagerEditor/Source/AssetManagerEditor/Private/SAssetAuditBrowser.cpp
+++ b/Engine/Plugins/Editor/AssetManagerEditor/Source/AssetManagerEditor/Private/SAssetAuditBrowser.cpp
@@ -435,6 +435,14 @@ void SAssetAuditBrowser::Construct(const FArguments& InArgs)
 					]
 					+SHorizontalBox::Slot()
 					.AutoWidth()
+					[
+						IAssetManagerEditorModule::MakePrimaryAssetTypeSelector(
+							FOnGetPrimaryAssetDisplayText::CreateLambda([] { return LOCTEXT("AddManagedAssetsOfType", "Add Managed Assets of Type"); }),
+							FOnSetPrimaryAssetType::CreateSP(this, &SAssetAuditBrowser::AddManagedAssetsOfType),
+							false, false)
+					]
+					+SHorizontalBox::Slot()
+					.AutoWidth()
 					[
 						SNew(SButton)
 						.HAlign(HAlign_Center)
@@ -786,6 +794,26 @@ void SAssetAuditBrowser::AddManagedAssets(FPrimaryAssetId AssetId)
 	}
 }
 
+void SAssetAuditBrowser::AddManagedAssetsOfType(FPrimaryAssetType AssetType)
+{
+	if (AssetType.IsValid())
+	{
+		TArray<FSoftObjectPath> AssetArray;
+		AssetManager->GetPrimaryAssetPathList(AssetType, AssetArray);
+		AddAssetsToList(AssetArray, false);
+
+		TArray<FName> AssetPackageArray;
+		TArray<FPrimaryAssetId> PrimaryAssetIds;
+		AssetManager->GetPrimaryAssetIdList(AssetType, PrimaryAssetIds);
+		for (const FPrimaryAssetId& PrimaryAssetId : PrimaryAssetIds)
+		{
+			// Calling GetManagedPackageList rather than AddManagedAssets to reduce UI update calls
+			AssetManager->GetManagedPackageList(PrimaryAssetId, AssetPackageArray);
+		}
+		AddAssetsToList(AssetPackageArray, false);
+	}
+}
+
 void SAssetAuditBrowser::AddAssetsOfClass(UClass* AssetClass)
 {
 	FSlateApplication::Get().DismissAllMenus();
diff --git a/Engine/Plugins/Editor/AssetManagerEditor/Source/AssetManagerEditor/Private/SAssetAuditBrowser.h b/Engine/Plugins/Editor/AssetManagerEditor/Source/AssetManagerEditor/Private/SAssetAuditBrowser.h
--- a/Engine/Plugins/Editor/AssetManagerEditor/Source/AssetManagerEditor/Private/SAssetAuditBrowser.h
+++ b/Engine/Plugins/Editor/AssetManagerEditor/Source/AssetManagerEditor/Private/SAssetAuditBrowser.h
@@ -118,6 +118,9 @@ protected:
 	/** Adds all managed assets by an id */
 	void AddManagedAssets(FPrimaryAssetId AssetId);
 
+	/** Adds all managed assets by by type. Useful to see how much size a primary asset type is using. */
+	void AddManagedAssetsOfType(FPrimaryAssetType AssetType);
+
 	/** Adds all assets of a class */
 	void AddAssetsOfClass(UClass* AssetClass);
 
