AssetManagerEditor: Remove debug code to identify whether pointer is dangling in callbacks (debug code was hit and it is indeed). Fix the dangling pointer by using CreateSPLambda instead of CreateLambda.
#rnx

diff --git a/Engine/Plugins/Editor/AssetManagerEditor/Source/AssetManagerEditor/Private/AssetSourceControlContextMenu.cpp b/Engine/Plugins/Editor/AssetManagerEditor/Source/AssetManagerEditor/Private/AssetSourceControlContextMenu.cpp
--- a/Engine/Plugins/Editor/AssetManagerEditor/Source/AssetManagerEditor/Private/AssetSourceControlContextMenu.cpp
+++ b/Engine/Plugins/Editor/AssetManagerEditor/Source/AssetManagerEditor/Private/AssetSourceControlContextMenu.cpp
@@ -117,8 +117,6 @@ namespace UE::AssetSourceControlContextMenu::Private
 class FAssetSourceControlContextMenuState : public TSharedFromThis<FAssetSourceControlContextMenuState>
 {
 public:
-	~FAssetSourceControlContextMenuState();
-
 	void Initialize(FToolMenuSection& InSection);
 
 	bool IsValid() const;
@@ -196,8 +194,6 @@ private:
 	/** Cancel any currently running perforce operation */
 	void CancelCacheCanExecuteVars();
 
-	bool IsClassStillConstructed() const;
-
 private:
 
 	TArray<FAssetData> SelectedAssets;
@@ -208,9 +204,6 @@ private:
 	TSharedPtr<class ISourceControlOperation, ESPMode::ThreadSafe> SCCOperation;
 	std::atomic<EAsyncState> AsyncState = EAsyncState::None;
 
-	static constexpr int32 ValiditySentryConstant = 0x85B921C3;
-	int32 ValiditySentry = ValiditySentryConstant;
-
 	// If folders are present in the selection then we adjust some behaviors to account for the potentially large amount of files to scan
 	bool bContainsFolders = false;
 
@@ -404,9 +397,9 @@ void FAssetSourceControlContextMenuState::FillSourceControlSubMenu(UToolMenu* Me
 			FSlateIcon(FRevisionControlStyleManager::GetStyleSetName(), "RevisionControl.Actions.Sync"),
 			FUIAction(
 				ExecutionCheck(FExecuteAction::CreateSP(this, &FAssetSourceControlContextMenuState::ExecuteSCCSync)),
-				FCanExecuteAction::CreateLambda([this]() { return IsActionEnabled(CanExecuteSCCSync()); })
+				FCanExecuteAction::CreateSPLambda(this, [this]() { return IsActionEnabled(CanExecuteSCCSync()); })
 			),
-			FIsAsyncProcessingActive::CreateLambda([this]() { return IsStillScanning(CanExecuteSCCSync()); })
+			FIsAsyncProcessingActive::CreateSPLambda(this, [this]() { return IsStillScanning(CanExecuteSCCSync()); })
 		);
 	}
 
@@ -415,49 +408,67 @@ void FAssetSourceControlContextMenuState::FillSourceControlSubMenu(UToolMenu* Me
 		AddAsyncMenuEntry(Section,
 			"SCCCheckOut",
 			LOCTEXT("SCCCheckOut", "Check Out"),
-			TAttribute<FText>::CreateLambda([this]()
+			TAttribute<FText>::CreateLambda([pWeakThis = this->AsWeak()]()
 				{
-					if (IsStillScanning(CanExecuteSCCCheckOut()) || CheckedOutUsers.Num() == 0)
+					TSharedPtr<FAssetSourceControlContextMenuState> pThis = pWeakThis.Pin();
+					if (!pThis) return FText();
+
+					if (pThis->IsStillScanning(pThis->CanExecuteSCCCheckOut()) || pThis->CheckedOutUsers.Num() == 0)
 					{
 						return LOCTEXT("SCCCheckOutTooltip", "Check out the selected assets from revision control.");
 					}
 
-					return FText::Format(LOCTEXT("SCCPartialCheckOut", "Checks out the selected assets from revision control that are not currently locked.\n\nLocked Assets:\n{0}"), CheckedOutUsersText);
+					return FText::Format(LOCTEXT("SCCPartialCheckOut",
+						"Checks out the selected assets from revision control that are not currently locked.\n\nLocked Assets:\n{0}"),
+						pThis->CheckedOutUsersText);
 				}),
 
-			TAttribute<FSlateIcon>::CreateLambda([this]()
+			TAttribute<FSlateIcon>::CreateLambda([pWeakThis = this->AsWeak()]()
 				{
-					return FSlateIcon(FRevisionControlStyleManager::GetStyleSetName(), CheckedOutUsers.Num() == 0 ? "RevisionControl.Actions.CheckOut" : "RevisionControl.Locked");
+					TSharedPtr<FAssetSourceControlContextMenuState> pThis = pWeakThis.Pin();
+					if (!pThis) return FSlateIcon();
+
+					return FSlateIcon(FRevisionControlStyleManager::GetStyleSetName(),
+						pThis->CheckedOutUsers.Num() ? "RevisionControl.Actions.CheckOut" : "RevisionControl.Locked");
 				}),
 					FUIAction(
 						ExecutionCheck(FExecuteAction::CreateSP(this, &FAssetSourceControlContextMenuState::ExecuteSCCCheckOut)),
-						FCanExecuteAction::CreateLambda([this]() { return IsActionEnabled(CanExecuteSCCCheckOut()); })
+						FCanExecuteAction::CreateSPLambda(this, [this]() { return IsActionEnabled(CanExecuteSCCCheckOut()); })
 					),
-					FIsAsyncProcessingActive::CreateLambda([this]() { return IsStillScanning(CanExecuteSCCCheckOut()); })
+					FIsAsyncProcessingActive::CreateSPLambda(this, [this]() { return IsStillScanning(CanExecuteSCCCheckOut()); })
 					);
 
 		AddAsyncMenuEntry(Section,
 			"SCCSyncAndCheckOut",
 			LOCTEXT("SCCSyncAndCheckOut", "Sync And Check Out"),
-			TAttribute<FText>::CreateLambda([this]()
+			TAttribute<FText>::CreateLambda([pWeakThis = this->AsWeak()]()
 				{
-					if (IsStillScanning(CanExecuteSCCCheckOut()) || CheckedOutUsers.Num() == 0)
+					TSharedPtr<FAssetSourceControlContextMenuState> pThis = pWeakThis.Pin();
+					if (!pThis) return FText();
+
+					if (pThis->IsStillScanning(pThis->CanExecuteSCCCheckOut()) || pThis->CheckedOutUsers.Num() == 0)
 					{
 						return LOCTEXT("SCCSyncAndCheckOutTooltip", "Sync to latest and Check out the selected assets from revision control.");
 					}
 
-					return FText::Format(LOCTEXT("SCCPartialSyncAndCheckOut", "Sync to latest and Checks out the selected assets from revision control that are not currently locked.\n\nLocked Assets:\n{0}"), CheckedOutUsersText);
+					return FText::Format(LOCTEXT("SCCPartialSyncAndCheckOut",
+						"Sync to latest and Checks out the selected assets from revision control that are not currently locked.\n\nLocked Assets:\n{0}"),
+						pThis->CheckedOutUsersText);
 				}),
 
-			TAttribute<FSlateIcon>::CreateLambda([this]()
+			TAttribute<FSlateIcon>::CreateLambda([pWeakThis = this->AsWeak()]()
 				{
-					return FSlateIcon(FRevisionControlStyleManager::GetStyleSetName(), CheckedOutUsers.Num() == 0 ? "RevisionControl.Actions.SyncAndCheckOut" : "RevisionControl.Locked");
+					TSharedPtr<FAssetSourceControlContextMenuState> pThis = pWeakThis.Pin();
+					if (!pThis) return FSlateIcon();
+
+					return FSlateIcon(FRevisionControlStyleManager::GetStyleSetName(),
+						pThis->CheckedOutUsers.Num() == 0 ? "RevisionControl.Actions.SyncAndCheckOut" : "RevisionControl.Locked");
 				}),
 					FUIAction(
 						ExecutionCheck(FExecuteAction::CreateSP(this, &FAssetSourceControlContextMenuState::ExecuteSCCSyncAndCheckOut)),
-						FCanExecuteAction::CreateLambda([this]() { return IsActionEnabled(CanExecuteSCCSyncAndCheckOut()); })
+						FCanExecuteAction::CreateSPLambda(this, [this]() { return IsActionEnabled(CanExecuteSCCSyncAndCheckOut()); })
 					),
-					FIsAsyncProcessingActive::CreateLambda([this]() { return IsStillScanning(CanExecuteSCCSyncAndCheckOut()); })
+					FIsAsyncProcessingActive::CreateSPLambda(this, [this]() { return IsStillScanning(CanExecuteSCCSyncAndCheckOut()); })
 					);
 	}
 
@@ -470,9 +481,9 @@ void FAssetSourceControlContextMenuState::FillSourceControlSubMenu(UToolMenu* Me
 			FSlateIcon(FRevisionControlStyleManager::GetStyleSetName(), "RevisionControl.Actions.MakeWritable"),
 			FUIAction(
 				ExecutionCheck(FExecuteAction::CreateSP(this, &FAssetSourceControlContextMenuState::ExecuteSCCMakeWritable)),
-				FCanExecuteAction::CreateLambda([this]() { return IsActionEnabled(CanExecuteSCCMakeWritable()); })
+				FCanExecuteAction::CreateSPLambda(this, [this]() { return IsActionEnabled(CanExecuteSCCMakeWritable()); })
 			),
-			FIsAsyncProcessingActive::CreateLambda([this]() { return IsStillScanning(CanExecuteSCCMakeWritable()); })
+			FIsAsyncProcessingActive::CreateSPLambda(this, [this]() { return IsStillScanning(CanExecuteSCCMakeWritable()); })
 		);
 	}
 
@@ -483,9 +494,9 @@ void FAssetSourceControlContextMenuState::FillSourceControlSubMenu(UToolMenu* Me
 		FSlateIcon(FRevisionControlStyleManager::GetStyleSetName(), "RevisionControl.Actions.Add"),
 		FUIAction(
 			ExecutionCheck(FExecuteAction::CreateSP(this, &FAssetSourceControlContextMenuState::ExecuteSCCOpenForAdd)),
-			FCanExecuteAction::CreateLambda([this]() { return IsActionEnabled(CanExecuteSCCOpenForAdd()); })
+			FCanExecuteAction::CreateSPLambda(this, [this]() { return IsActionEnabled(CanExecuteSCCOpenForAdd()); })
 		),
-		FIsAsyncProcessingActive::CreateLambda([this]() { return IsStillScanning(CanExecuteSCCOpenForAdd()); })
+		FIsAsyncProcessingActive::CreateSPLambda(this, [this]() { return IsStillScanning(CanExecuteSCCOpenForAdd()); })
 	);
 
 	if (!bUsesSnapshots)
@@ -497,9 +508,9 @@ void FAssetSourceControlContextMenuState::FillSourceControlSubMenu(UToolMenu* Me
 			FSlateIcon(FRevisionControlStyleManager::GetStyleSetName(), "RevisionControl.Actions.Submit"),
 			FUIAction(
 				ExecutionCheck(FExecuteAction::CreateSP(this, &FAssetSourceControlContextMenuState::ExecuteSCCCheckIn)),
-				FCanExecuteAction::CreateLambda([this]() { return IsActionEnabled(CanExecuteSCCCheckIn()); })
+				FCanExecuteAction::CreateSPLambda(this, [this]() { return IsActionEnabled(CanExecuteSCCCheckIn()); })
 			),
-			FIsAsyncProcessingActive::CreateLambda([this]() { return IsStillScanning(CanExecuteSCCCheckIn()); })
+			FIsAsyncProcessingActive::CreateSPLambda(this, [this]() { return IsStillScanning(CanExecuteSCCCheckIn()); })
 		);
 	}
 
@@ -510,9 +521,9 @@ void FAssetSourceControlContextMenuState::FillSourceControlSubMenu(UToolMenu* Me
 		FSlateIcon(FRevisionControlStyleManager::GetStyleSetName(), "RevisionControl.Actions.History"),
 		FUIAction(
 			FExecuteAction::CreateSP(this, &FAssetSourceControlContextMenuState::ExecuteSCCHistory),
-			FCanExecuteAction::CreateLambda([this]() { return IsActionEnabled(CanExecuteSCCHistory()); })
+			FCanExecuteAction::CreateSPLambda(this, [this]() { return IsActionEnabled(CanExecuteSCCHistory()); })
 		),
-		FIsAsyncProcessingActive::CreateLambda([this]() { return IsStillScanning(CanExecuteSCCHistory()); })
+		FIsAsyncProcessingActive::CreateSPLambda(this, [this]() { return IsStillScanning(CanExecuteSCCHistory()); })
 	);
 
 	if (bUsesDiffAgainstDepot)
@@ -524,9 +535,9 @@ void FAssetSourceControlContextMenuState::FillSourceControlSubMenu(UToolMenu* Me
 			FSlateIcon(FRevisionControlStyleManager::GetStyleSetName(), "RevisionControl.Actions.Diff"),
 			FUIAction(
 				FExecuteAction::CreateSP(this, &FAssetSourceControlContextMenuState::ExecuteSCCDiffAgainstDepot),
-				FCanExecuteAction::CreateLambda([this]() { return IsActionEnabled(CanExecuteSCCDiffAgainstDepot()); })
+				FCanExecuteAction::CreateSPLambda(this, [this]() { return IsActionEnabled(CanExecuteSCCDiffAgainstDepot()); })
 			),
-			FIsAsyncProcessingActive::CreateLambda([this]() { return IsStillScanning(CanExecuteSCCDiffAgainstDepot()); })
+			FIsAsyncProcessingActive::CreateSPLambda(this, [this]() { return IsStillScanning(CanExecuteSCCDiffAgainstDepot()); })
 		);
 	}
 
@@ -537,9 +548,9 @@ void FAssetSourceControlContextMenuState::FillSourceControlSubMenu(UToolMenu* Me
 		FSlateIcon(FRevisionControlStyleManager::GetStyleSetName(), "RevisionControl.Actions.Revert"),
 		FUIAction(
 			FExecuteAction::CreateSP(this, &FAssetSourceControlContextMenuState::ExecuteSCCRevert),
-			FCanExecuteAction::CreateLambda([this]() { return IsActionEnabled(CanExecuteSCCRevert()); })
+			FCanExecuteAction::CreateSPLambda(this, [this]() { return IsActionEnabled(CanExecuteSCCRevert()); })
 		),
-		FIsAsyncProcessingActive::CreateLambda([this]() { return IsStillScanning(CanExecuteSCCRevert()); })
+		FIsAsyncProcessingActive::CreateSPLambda(this, [this]() { return IsStillScanning(CanExecuteSCCRevert()); })
 	);
 
 	if (bUsesReadOnly)
@@ -551,9 +562,9 @@ void FAssetSourceControlContextMenuState::FillSourceControlSubMenu(UToolMenu* Me
 			FSlateIcon(FRevisionControlStyleManager::GetStyleSetName(), "RevisionControl.Actions.Revert"),
 			FUIAction(
 				ExecutionCheck(FExecuteAction::CreateSP(this, &FAssetSourceControlContextMenuState::ExecuteSCCRevertWritable)),
-				FCanExecuteAction::CreateLambda([this]() { return IsActionEnabled(CanExecuteSCCRevertWritable()); })
+				FCanExecuteAction::CreateSPLambda(this, [this]() { return IsActionEnabled(CanExecuteSCCRevertWritable()); })
 			),
-			FIsAsyncProcessingActive::CreateLambda([this]() { return IsStillScanning(CanExecuteSCCRevertWritable()); })
+			FIsAsyncProcessingActive::CreateSPLambda(this, [this]() { return IsStillScanning(CanExecuteSCCRevertWritable()); })
 		);
 	}
 
@@ -567,23 +578,8 @@ void FAssetSourceControlContextMenuState::FillSourceControlSubMenu(UToolMenu* Me
 			FSlateIcon(FRevisionControlStyleManager::GetStyleSetName(), "RevisionControl.Actions.Merge"),
 			FUIAction(
 				FExecuteAction::CreateSP(this, &FAssetSourceControlContextMenuState::ExecuteSCCMerge),
-				FCanExecuteAction::CreateLambda([this]() {
-					if (!IsClassStillConstructed())
-					{
-						UE_LOG(LogCore, Fatal, TEXT("Dangling pointer detected in IsActionEnabled lambda passed to FUIAction"));
-						return false;
-					}
-					return IsActionEnabled(CanExecuteSCCMerge());
-				})),
-			FIsAsyncProcessingActive::CreateLambda([this]()
-				{
-					if (!IsClassStillConstructed())
-					{
-						UE_LOG(LogCore, Fatal, TEXT("Dangling pointer detected in IsAsyncProcessingActive lambda passed to AddAsyncMenuEntry"));
-						return false;
-					}
-					return IsStillScanning(CanExecuteSCCMerge());
-				})
+				FCanExecuteAction::CreateSPLambda(this, [this]() { return IsActionEnabled(CanExecuteSCCMerge()); })),
+			FIsAsyncProcessingActive::CreateSPLambda(this, [this]() { return IsStillScanning(CanExecuteSCCMerge()); })
 		);
 	}
 
@@ -601,7 +597,7 @@ void FAssetSourceControlContextMenuState::FillSourceControlSubMenu(UToolMenu* Me
 
 FExecuteAction FAssetSourceControlContextMenuState::ExecutionCheck(FExecuteAction&& InAction) const
 {
-	return FExecuteAction::CreateLambda([this, Action = MoveTemp(InAction)]()
+	return FExecuteAction::CreateSPLambda(this, [this, Action = MoveTemp(InAction)]()
 		{
 			if (SelectedAssets.Num() > 10) // Todo: Make this into a user option, or different values per operation type
 			{
@@ -889,19 +885,6 @@ void FAssetSourceControlContextMenuState::ExecuteSCCSync() const
 	AssetViewUtils::SyncPackagesFromSourceControl(PackageNames);
 }
 
-FAssetSourceControlContextMenuState::~FAssetSourceControlContextMenuState()
-{
-	ValiditySentry = 0;
-}
-
-bool FAssetSourceControlContextMenuState::IsClassStillConstructed() const
-{
-	// UE-195836. We suspect that CanExecuteSCCMerge is being called on a dangling pointer
-	// so we have added a sentry value in the class that gets overwritten with a different 
-	// value in the destructor.
-	return ValiditySentry == ValiditySentryConstant;
-}
-
 bool FAssetSourceControlContextMenuState::CanExecuteSCCMerge() const
 {
 	FAssetToolsModule& AssetToolsModule = FModuleManager::GetModuleChecked<FAssetToolsModule>("AssetTools");
