PR #10755: [5.3][Plastic SCM] Fix context Check In from content browser
PR #10378 [Source Control][5.2] Fix the Check In option in the context menu for Git & Plastic SCM

diff --git a/Engine/Plugins/Developer/PlasticSourceControl/Source/PlasticSourceControl/Private/PlasticSourceControlProvider.cpp b/Engine/Plugins/Developer/PlasticSourceControl/Source/PlasticSourceControl/Private/PlasticSourceControlProvider.cpp
--- a/Engine/Plugins/Developer/PlasticSourceControl/Source/PlasticSourceControl/Private/PlasticSourceControlProvider.cpp
+++ b/Engine/Plugins/Developer/PlasticSourceControl/Source/PlasticSourceControl/Private/PlasticSourceControlProvider.cpp
@@ -501,9 +501,6 @@ bool FPlasticSourceControlProvider::UsesCheckout() const
 
 bool FPlasticSourceControlProvider::UsesFileRevisions() const
 {
-	// This API introduced in UE5.1 is still broken as of UE5.3
-	// (preventing the user to use the source control context menu for checkin if returning false)
-	// return IsPartialWorkspace();
 	return true;
 }
 
diff --git a/Engine/Plugins/Editor/AssetManagerEditor/Source/AssetManagerEditor/Private/AssetSourceControlContextMenu.cpp b/Engine/Plugins/Editor/AssetManagerEditor/Source/AssetManagerEditor/Private/AssetSourceControlContextMenu.cpp
--- a/Engine/Plugins/Editor/AssetManagerEditor/Source/AssetManagerEditor/Private/AssetSourceControlContextMenu.cpp
+++ b/Engine/Plugins/Editor/AssetManagerEditor/Source/AssetManagerEditor/Private/AssetSourceControlContextMenu.cpp
@@ -1080,6 +1080,7 @@ void FAssetSourceControlContextMenuState::TryCacheCanExecuteVars(const TArray<FS
 
 	const bool bUsesCheckout = SourceControlProvider.UsesCheckout();
 	const bool bUsesFileRevisions = SourceControlProvider.UsesFileRevisions();
+	const bool bUsesSnapshots = SourceControlProvider.UsesSnapshots();
 
 	// If a package is dirty, allow a revert of the in-memory changes that have not yet been saved to disk.
 	if (AllowExecuteSCCRevertUnsaved())
@@ -1132,7 +1133,10 @@ void FAssetSourceControlContextMenuState::TryCacheCanExecuteVars(const TArray<FS
 				{
 					bCanExecuteSCCSync = bCanSyncCurrentItem = true;
 				}
+			}
 
+			if (!bUsesSnapshots)
+			{
 				if (SourceControlState->CanCheckIn())
 				{
 					bCanExecuteSCCCheckIn = true;
