Removing deprecated method calls

diff --git a/Engine/Plugins/Editor/AssetManagerEditor/Source/AssetManagerEditor/Private/AssetManagerEditorModule.cpp b/Engine/Plugins/Editor/AssetManagerEditor/Source/AssetManagerEditor/Private/AssetManagerEditorModule.cpp
--- a/Engine/Plugins/Editor/AssetManagerEditor/Source/AssetManagerEditor/Private/AssetManagerEditorModule.cpp
+++ b/Engine/Plugins/Editor/AssetManagerEditor/Source/AssetManagerEditor/Private/AssetManagerEditorModule.cpp
@@ -2479,10 +2479,9 @@ bool FAssetManagerEditorModule::WriteCollection(FName CollectionName, ECollectio
 		ResultsMessage = FText::Format(LOCTEXT("NothingToAddToCollection", "Nothing to add to collection {0}"), FText::FromName(CollectionName));
 	}
 	else if (CreateOrEmptyCollection(CollectionName, ShareType))
-	{	
-PRAGMA_DISABLE_DEPRECATION_WARNINGS
-		if (CollectionManager.AddToCollection(CollectionName, ECollectionShareType::CST_Local, UE::SoftObjectPath::Private::ConvertSoftObjectPaths(ObjectPathsToAddToCollection.Array())))
-PRAGMA_ENABLE_DEPRECATION_WARNINGS
+	{
+		if (CollectionManager.AddToCollection(
+				CollectionName, ECollectionShareType::CST_Local, ObjectPathsToAddToCollection.Array()))
 		{
 			UE_LOG(LogAssetManagerEditor, Log, TEXT("Updated collection %s"), *CollectionName.ToString());
 			ResultsMessage = FText::Format(LOCTEXT("CreateCollectionSucceeded", "Updated collection {0}"), FText::FromName(CollectionName));
diff --git a/Engine/Source/Editor/ContentBrowser/Private/AssetViewWidgets.cpp b/Engine/Source/Editor/ContentBrowser/Private/AssetViewWidgets.cpp
--- a/Engine/Source/Editor/ContentBrowser/Private/AssetViewWidgets.cpp
+++ b/Engine/Source/Editor/ContentBrowser/Private/AssetViewWidgets.cpp
@@ -1035,9 +1035,8 @@ TSharedRef<SWidget> SAssetViewItem::CreateToolTipWidget() const
 				ICollectionManager& CollectionManager = FCollectionManagerModule::GetModule().Get();
 
 				TArray<FCollectionNameType> CollectionsContainingObject;
-PRAGMA_DISABLE_DEPRECATION_WARNINGS
-				CollectionManager.GetCollectionsContainingObject(ItemAssetData.ObjectPath, CollectionsContainingObject);
-PRAGMA_ENABLE_DEPRECATION_WARNINGS
+				CollectionManager.GetCollectionsContainingObject(
+					ItemAssetData.ToSoftObjectPath(), CollectionsContainingObject);
 
 				if (CollectionsContainingObject.Num() > 0)
 				{
diff --git a/Engine/Source/Editor/ContentBrowser/Private/SCollectionView.cpp b/Engine/Source/Editor/ContentBrowser/Private/SCollectionView.cpp
--- a/Engine/Source/Editor/ContentBrowser/Private/SCollectionView.cpp
+++ b/Engine/Source/Editor/ContentBrowser/Private/SCollectionView.cpp
@@ -1348,9 +1348,8 @@ FReply SCollectionView::HandleDragDropOnCollectionItem(TSharedRef<FCollectionIte
 		const double BeginTimeSec = FPlatformTime::Seconds();
 		int32 NumAdded = 0;
 		FText Message;
-PRAGMA_DISABLE_DEPRECATION_WARNINGS
-		if (CollectionManagerModule.Get().AddToCollection(CollectionItem->CollectionName, CollectionItem->CollectionType, UE::SoftObjectPath::Private::ConvertSoftObjectPaths(ObjectPaths), &NumAdded))
-PRAGMA_ENABLE_DEPRECATION_WARNINGS
+		if (CollectionManagerModule.Get().AddToCollection(
+				CollectionItem->CollectionName, CollectionItem->CollectionType, ObjectPaths, &NumAdded))
 		{
 			if (DroppedAssets.Num() == 1)
 			{
diff --git a/Engine/Source/Editor/UnrealEd/Private/Commandlets/DumpMaterialShaderTypes.cpp b/Engine/Source/Editor/UnrealEd/Private/Commandlets/DumpMaterialShaderTypes.cpp
--- a/Engine/Source/Editor/UnrealEd/Private/Commandlets/DumpMaterialShaderTypes.cpp
+++ b/Engine/Source/Editor/UnrealEd/Private/Commandlets/DumpMaterialShaderTypes.cpp
@@ -892,7 +892,9 @@ int32 UDumpMaterialShaderTypesCommandlet::Main(const FString& Params)
 
 				FCollectionManagerModule& CollectionManagerModule = FCollectionManagerModule::GetModule();
 PRAGMA_DISABLE_DEPRECATION_WARNINGS
-				CollectionManagerModule.Get().GetObjectsInCollection(FName(*CollectionName), ECollectionShareType::CST_All, Filter.ObjectPaths, ECollectionRecursionFlags::SelfAndChildren);
+CollectionManagerModule.Get().GetObjectsInCollection(
+	FName(*CollectionName), ECollectionShareType::CST_All, Filter.SoftObjectPaths,
+	ECollectionRecursionFlags::SelfAndChildren);
 PRAGMA_ENABLE_DEPRECATION_WARNINGS
 
 				AssetRegistry.GetAssets(Filter, MaterialList);
diff --git a/Engine/Source/Editor/UnrealEd/Private/Commandlets/ParticleSystemAuditCommandlet.cpp b/Engine/Source/Editor/UnrealEd/Private/Commandlets/ParticleSystemAuditCommandlet.cpp
--- a/Engine/Source/Editor/UnrealEd/Private/Commandlets/ParticleSystemAuditCommandlet.cpp
+++ b/Engine/Source/Editor/UnrealEd/Private/Commandlets/ParticleSystemAuditCommandlet.cpp
@@ -81,9 +81,9 @@ bool UParticleSystemAuditCommandlet::ProcessParticleSystems()
 	if (!FilterCollection.IsEmpty())
 	{
 		FCollectionManagerModule& CollectionManagerModule = FCollectionManagerModule::GetModule();
-PRAGMA_DISABLE_DEPRECATION_WARNINGS
-		CollectionManagerModule.Get().GetObjectsInCollection(FName(*FilterCollection), ECollectionShareType::CST_All, Filter.ObjectPaths, ECollectionRecursionFlags::SelfAndChildren);
-PRAGMA_ENABLE_DEPRECATION_WARNINGS
+		CollectionManagerModule.Get().GetObjectsInCollection(
+			FName(*FilterCollection), ECollectionShareType::CST_All, Filter.SoftObjectPaths,
+			ECollectionRecursionFlags::SelfAndChildren);
 	}
 
 	TArray<FAssetData> AssetList;
