[Reference Viewer] Editor Preferences for Ref Viewer no longer respected

Reference viewer widget is now retrieving and applying project settings for the following properties when is being constructed:

- Breadth Limit value
- Show Searchable Names on/off

diff --git a/Engine/Plugins/Editor/AssetManagerEditor/Source/AssetManagerEditor/Private/ReferenceViewer/SReferenceViewer.cpp b/Engine/Plugins/Editor/AssetManagerEditor/Source/AssetManagerEditor/Private/ReferenceViewer/SReferenceViewer.cpp
--- a/Engine/Plugins/Editor/AssetManagerEditor/Source/AssetManagerEditor/Private/ReferenceViewer/SReferenceViewer.cpp
+++ b/Engine/Plugins/Editor/AssetManagerEditor/Source/AssetManagerEditor/Private/ReferenceViewer/SReferenceViewer.cpp
@@ -22,6 +22,7 @@
 #include "AssetManagerEditorCommands.h"
 #include "EditorWidgetsModule.h"
 #include "ReferenceViewer/ReferenceViewerSettings.h"
+#include "Settings/EditorProjectSettings.h"
 #include "Toolkits/GlobalEditorCommonCommands.h"
 #include "Engine/AssetManager.h"
 #include "ReferenceViewer/SReferenceViewerFilterBar.h"
@@ -117,6 +118,37 @@ void SReferenceViewer::Construct(const FArguments& InArgs)
 	bShowCompactMode = true;
 	bDirtyResults = false;
 
+	// Retrieve and apply Breadth limit and show searchable names values from Project Settings
+	if (const UEditorProjectAppearanceSettings* DefaultProjectAppearanceSettings = GetDefault<UEditorProjectAppearanceSettings>())
+	{
+		FixAndHideSearchBreadthLimit = DefaultProjectAppearanceSettings->ReferenceViewerDefaultMaxSearchBreadth;
+
+		switch (DefaultProjectAppearanceSettings->ShowSearchableNames)
+		{
+			case EReferenceViewerSettingMode::NoPreference:
+				bShowShowSearchableNames = true;
+				break;
+
+			case EReferenceViewerSettingMode::ShowByDefault:
+				bShowShowSearchableNames = true;
+				break;
+
+			case EReferenceViewerSettingMode::HideByDefault:
+				bShowShowSearchableNames = false;
+				break;
+
+			default:
+				bShowShowSearchableNames = true;
+				break;
+		}
+
+		if (Settings)
+		{
+			Settings->SetSearchBreadthLimit(FixAndHideSearchBreadthLimit);
+			Settings->SetShowSearchableNames(bShowShowSearchableNames);
+		}
+	}
+
 	SAssignNew(FilterWidget, SReferenceViewerFilterBar)
 		.Visibility_Lambda([this]() { return !Settings->GetFiltersEnabled() ? EVisibility::Collapsed : EVisibility::Visible; })
 		.OnConvertItemToAssetData_Lambda([this] (FReferenceNodeInfo& InNodeInfo, FAssetData& OutAssetData) -> bool { 
