[Reference Viewer] ReferenceViewer shouldn't show the Collapsed Node if it only contains 1 reference

Instead of creating a collapsed node encompassing a single node, the node will be directly shown.

diff --git a/Engine/Plugins/Editor/AssetManagerEditor/Source/AssetManagerEditor/Private/ReferenceViewer/EdGraph_ReferenceViewer.cpp b/Engine/Plugins/Editor/AssetManagerEditor/Source/AssetManagerEditor/Private/ReferenceViewer/EdGraph_ReferenceViewer.cpp
--- a/Engine/Plugins/Editor/AssetManagerEditor/Source/AssetManagerEditor/Private/ReferenceViewer/EdGraph_ReferenceViewer.cpp
+++ b/Engine/Plugins/Editor/AssetManagerEditor/Source/AssetManagerEditor/Private/ReferenceViewer/EdGraph_ReferenceViewer.cpp
@@ -952,30 +952,53 @@ UEdGraphNode_Reference* UEdGraph_ReferenceViewer::RecursivelyCreateNodes(bool bI
 		// There were more references than allowed to be displayed. Make a collapsed node.
 		if (NodeInfo.OverflowCount > 0)
 		{
-			UEdGraphNode_Reference* OverflowNode = CreateReferenceNode();
+			UEdGraphNode_Reference* OverflowNode = nullptr;
 			FIntPoint RefNodeLoc;
 			RefNodeLoc.X = ChildLoc.X;
 			RefNodeLoc.Y = ChildLoc.Y;
 
-			if ( ensure(OverflowNode) )
+			// Overflow count is 1: instead of collapsing a single node, we can directly display it
+			if (NodeInfo.OverflowCount == 1)
 			{
-				OverflowNode->SetAllowThumbnail(!Settings->IsCompactMode());
+				// Reaching the overflowing node
+				if (NodeInfo.Children.IsValidIndex(Breadth))
+				{
+					const TPair<FAssetIdentifier, EDependencyPinCategory>& OverflowNodePair = NodeInfo.Children[Breadth];
 
-				TArray<FAssetIdentifier> CollapsedNodeIdentifiers;
-				for (; ChildIdx < InNodeInfos[InAssetId].Children.Num(); ChildIdx++)
+					const FAssetIdentifier& OverflowNodeAssetId = OverflowNodePair.Key;
+					OverflowNode = RecursivelyCreateNodes(bInReferencers, OverflowNodeAssetId, ChildLoc, NodeInfo.AssetId, NewNode, InNodeInfos, InCurrentDepth + 1, InMaxDepth);
+				}
+			}
+
+			// OverflowNode is not valid. Either NodeInfo.OverflowCount is > 1, or single node creation failed.
+			// Let's create a collapsed node.
+			if (!OverflowNode)
+			{
+				if (UEdGraphNode_Reference* CollapsedNode = CreateReferenceNode())
 				{
-					const TPair<FAssetIdentifier, EDependencyPinCategory>& Pair = InNodeInfos[InAssetId].Children[ChildIdx];
-					CollapsedNodeIdentifiers.Add(Pair.Key);
+					TArray<FAssetIdentifier> CollapsedNodeIdentifiers;
+					for (; ChildIdx < InNodeInfos[InAssetId].Children.Num(); ChildIdx++)
+					{
+						const TPair<FAssetIdentifier, EDependencyPinCategory>& Pair = InNodeInfos[InAssetId].Children[ChildIdx];
+						CollapsedNodeIdentifiers.Add(Pair.Key);
+					}
+
+					CollapsedNode->SetReferenceNodeCollapsed(RefNodeLoc, NodeInfo.OverflowCount, CollapsedNodeIdentifiers);
+					OverflowNode = CollapsedNode;
 				}
-				OverflowNode->SetReferenceNodeCollapsed(RefNodeLoc, NodeInfo.OverflowCount, CollapsedNodeIdentifiers);
+			}
 
-				if ( bInReferencers )
+			if (ensure(OverflowNode))
+			{
+				OverflowNode->SetAllowThumbnail(!Settings->IsCompactMode());
+
+				if (bInReferencers)
 				{
-					NewNode->AddReferencer( OverflowNode );
+					NewNode->AddReferencer(OverflowNode);
 				}
 				else
 				{
-					OverflowNode->AddReferencer( NewNode );
+					OverflowNode->AddReferencer(NewNode);
 				}
 			}
 		}
