[Reference Viewer] Collapsed nodes need to look more different

Matching the new design for Collapsed Nodes, which helps differentiating them in the graph, especially when zooming out

Reference design: https://jira.it.epicgames.com/browse/UE-193616

diff --git a/Engine/Plugins/Editor/AssetManagerEditor/Source/AssetManagerEditor/Private/ReferenceViewer/EdGraphNode_Reference.cpp b/Engine/Plugins/Editor/AssetManagerEditor/Source/AssetManagerEditor/Private/ReferenceViewer/EdGraphNode_Reference.cpp
--- a/Engine/Plugins/Editor/AssetManagerEditor/Source/AssetManagerEditor/Private/ReferenceViewer/EdGraphNode_Reference.cpp
+++ b/Engine/Plugins/Editor/AssetManagerEditor/Source/AssetManagerEditor/Private/ReferenceViewer/EdGraphNode_Reference.cpp
@@ -127,7 +127,8 @@ void UEdGraphNode_Reference::SetReferenceNodeCollapsed(const FIntPoint& NodeLoc,
 	bIsOverflow = true;
 	AssetBrush = FSlateIcon(FAppStyle::GetAppStyleSetName(), "Icons.WarningWithColor");
 
-	NodeTitle = FText::Format( LOCTEXT("ReferenceNodeCollapsedTitle", "{0} Collapsed nodes"), FText::AsNumber(InNumReferencesExceedingMax));
+	NodeTitle = FText::Format( LOCTEXT("ReferenceNodeCollapsedTitle", "{0}"), FText::AsNumber(InNumReferencesExceedingMax));
+
 	CacheAssetData(FAssetData());
 	AllocateDefaultPins();
 }
@@ -202,7 +203,16 @@ FLinearColor UEdGraphNode_Reference::GetNodeTitleColor() const
 
 FText UEdGraphNode_Reference::GetTooltipText() const
 {
+	// Showing up to 15 nodes paths, in order to avoid an extremely long tooltip and huge widget
+	constexpr int32 MaxReferenceNum = 15;
+
 	FString TooltipString;
+	int32 ReferenceCount = 0;
+
+	if (IsCollapsed())
+	{
+		TooltipString.Append(GetNodeTitle(ENodeTitleType::FullTitle).ToString() + " collapsed nodes:\n");
+	}
 	for (const FAssetIdentifier& AssetId : Identifiers)
 	{
 		if (!TooltipString.IsEmpty())
@@ -210,6 +220,15 @@ FText UEdGraphNode_Reference::GetTooltipText() const
 			TooltipString.Append(TEXT("\n"));
 		}
 		TooltipString.Append(AssetId.ToString());
+
+		ReferenceCount++;
+
+		// Avoiding an extremely long tooltip
+		if (ReferenceCount > MaxReferenceNum)
+		{
+			TooltipString.Append(TEXT("\n..."));
+			break;
+		}
 	}
 	return FText::FromString(TooltipString);
 }
diff --git a/Engine/Plugins/Editor/AssetManagerEditor/Source/AssetManagerEditor/Private/ReferenceViewer/SReferenceNode.cpp b/Engine/Plugins/Editor/AssetManagerEditor/Source/AssetManagerEditor/Private/ReferenceViewer/SReferenceNode.cpp
--- a/Engine/Plugins/Editor/AssetManagerEditor/Source/AssetManagerEditor/Private/ReferenceViewer/SReferenceNode.cpp
+++ b/Engine/Plugins/Editor/AssetManagerEditor/Source/AssetManagerEditor/Private/ReferenceViewer/SReferenceNode.cpp
@@ -14,11 +14,21 @@
 
 #define LOCTEXT_NAMESPACE "ReferenceViewer"
 
+namespace FReferenceNode::Private
+{
+	static constexpr int32 ThumbnailSize = 128;
+	static FVector2D CollapsedNodeTextBoxScale(ThumbnailSize * 0.7f, ThumbnailSize * 0.6f);
+}
+
 void SReferenceNode::Construct( const FArguments& InArgs, UEdGraphNode_Reference* InNode )
 {
-	const int32 ThumbnailSize = 128;
+	using namespace FReferenceNode::Private;
 
-	if (InNode->AllowsThumbnail())
+	if (InNode->IsCollapsed())
+	{
+		bIsCollapsed = true;
+	}
+	else if (InNode->AllowsThumbnail())
 	{
 		if (InNode->UsesThumbnail())
 		{
@@ -26,7 +36,7 @@ void SReferenceNode::Construct( const FArguments& InArgs, UEdGraphNode_Reference
 			TSharedPtr<FAssetThumbnailPool> AssetThumbnailPool = InNode->GetReferenceViewerGraph()->GetAssetThumbnailPool();
 			AssetThumbnail = MakeShareable( new FAssetThumbnail( InNode->GetAssetData(), ThumbnailSize, ThumbnailSize, AssetThumbnailPool ) );
 		}
-		else if (InNode->IsPackage() || InNode->IsCollapsed())
+		else if (InNode->IsPackage())
 		{
 			// Just make a generic thumbnail
 			AssetThumbnail = MakeShareable( new FAssetThumbnail( InNode->GetAssetData(), ThumbnailSize, ThumbnailSize, NULL ) );
@@ -72,21 +82,45 @@ void SReferenceNode::UpdateGraphNode()
 		IconBrush = GraphNode->GetIconAndTint(IconColor).GetOptionalIcon();
 	}
 
-	TSharedRef<SWidget> ThumbnailWidget = SNullWidget::NullWidget;
+	TSharedRef<SWidget> MainNodeWidget = SNullWidget::NullWidget;
 	UEdGraphNode_Reference* RefGraphNode = CastChecked<UEdGraphNode_Reference>(GraphNode);
 	bool bIsADuplicate = RefGraphNode->IsADuplicate();
 
 	FLinearColor OpacityColor = RefGraphNode->GetIsFiltered() ? FLinearColor(1.0, 1.0, 1.0, 0.4) : FLinearColor::White;
-	
-	if ( AssetThumbnail.IsValid() )
+
+	if (bIsCollapsed)
 	{
+		using namespace FReferenceNode::Private;
+
+		MainNodeWidget =
+			SNew(SBox)
+			.MinDesiredWidth(CollapsedNodeTextBoxScale.X)
+			.MinDesiredHeight(CollapsedNodeTextBoxScale.Y)
+			.HAlign(HAlign_Center)
+			.VAlign(VAlign_Center)
+			[
+				SNew(STextBlock)
+				.Font(FSlateFontInfo(FCoreStyle::GetDefaultFont(), 30))
+				.ColorAndOpacity(FColor::White)
+				.Text_Lambda([WeakGraphNode = TWeakObjectPtr<UEdGraphNode>(GraphNode)]()
+				{
+					if (UEdGraphNode* GraphNode = WeakGraphNode.Get())
+					{
+						return GraphNode->GetNodeTitle(ENodeTitleType::FullTitle);
+					}
 
+					return LOCTEXT("InvalidNodeTitle", "Invalid Node");
+				})
+			];
+	}
+	else if (AssetThumbnail.IsValid())
+	{
 		FAssetThumbnailConfig ThumbnailConfig;
 		ThumbnailConfig.bAllowFadeIn = RefGraphNode->UsesThumbnail();
 		ThumbnailConfig.bForceGenericThumbnail = !RefGraphNode->UsesThumbnail();
 		ThumbnailConfig.AssetTypeColorOverride = FLinearColor::Transparent;
 
-		ThumbnailWidget =
+		MainNodeWidget =
 			SNew(SBox)
 			.WidthOverride(AssetThumbnail->GetSize().X)
 			.HeightOverride(AssetThumbnail->GetSize().Y)
@@ -131,6 +165,7 @@ void SReferenceNode::UpdateGraphNode()
 				.Padding(0)
 				[
 					SNew(SBorder)
+					.Visibility(this, &SReferenceNode::GetCollapsedVisibility)
 					.BorderImage( FReferenceViewerStyle::Get().GetBrush("Graph.Node.ColorSpill") )
 					.Padding( FMargin(10.0f, 4.0f, 6.0f, 4.0f) )
 					.BorderBackgroundColor( this, &SReferenceNode::GetNodeTitleColor )
@@ -220,14 +255,14 @@ void SReferenceNode::UpdateGraphNode()
 								SAssignNew(LeftNodeBox, SVerticalBox)
 							]
 						]
-						
+
 						+SHorizontalBox::Slot()
 						.VAlign(VAlign_Center)
 						.HAlign(HAlign_Center)
 						.FillWidth(1.0f)
 						[
 							// Thumbnail
-							ThumbnailWidget
+							MainNodeWidget
 						]
 
 						+SHorizontalBox::Slot()
@@ -276,4 +311,9 @@ void SReferenceNode::UpdateGraphNode()
 	CreatePinWidgets();
 }
 
+EVisibility SReferenceNode::GetCollapsedVisibility() const
+{
+	return bIsCollapsed ? EVisibility::Collapsed : EVisibility::Visible;
+}
+
 #undef LOCTEXT_NAMESPACE
diff --git a/Engine/Plugins/Editor/AssetManagerEditor/Source/AssetManagerEditor/Private/ReferenceViewer/SReferenceNode.h b/Engine/Plugins/Editor/AssetManagerEditor/Source/AssetManagerEditor/Private/ReferenceViewer/SReferenceNode.h
--- a/Engine/Plugins/Editor/AssetManagerEditor/Source/AssetManagerEditor/Private/ReferenceViewer/SReferenceNode.h
+++ b/Engine/Plugins/Editor/AssetManagerEditor/Source/AssetManagerEditor/Private/ReferenceViewer/SReferenceNode.h
@@ -26,5 +26,12 @@ public:
 	// End SGraphNode implementation
 
 private:
+
+	/** Returns visibility setting for elements which should be hidden if node is collapsed */
+	EVisibility GetCollapsedVisibility() const;
+
 	TSharedPtr<class FAssetThumbnail> AssetThumbnail;
+
+	/** Whether the current node is collapsed (representing more than one node) or not */
+	bool bIsCollapsed = false;
 };
