Optimize deduplicating asset packages.

diff --git a/Engine/Plugins/Editor/AssetManagerEditor/Source/AssetManagerEditor/Private/AssetManagerEditorModule.cpp b/Engine/Plugins/Editor/AssetManagerEditor/Source/AssetManagerEditor/Private/AssetManagerEditorModule.cpp
--- a/Engine/Plugins/Editor/AssetManagerEditor/Source/AssetManagerEditor/Private/AssetManagerEditorModule.cpp
+++ b/Engine/Plugins/Editor/AssetManagerEditor/Source/AssetManagerEditor/Private/AssetManagerEditorModule.cpp
@@ -1,6 +1,9 @@
 // Copyright Epic Games, Inc. All Rights Reserved.
 
 #include "AssetManagerEditorModule.h"
+
+#include "Algo/Sort.h"
+#include "Algo/Unique.h"
 #include "AssetRegistry/AssetRegistryModule.h"
 #include "ContentBrowserDataLegacyBridge.h"
 #include "HAL/PlatformFile.h"
@@ -869,7 +872,6 @@ TArray<FName> FAssetManagerEditorModule::GetLevelEditorSelectedAssetPackages()
 
 TArray<FName> FAssetManagerEditorModule::GetContentBrowserSelectedAssetPackages(FOnContentBrowserGetSelection GetSelectionDelegate)
 {
-	TArray<FName> OutAssetPackages;
 	TArray<FAssetData> SelectedAssets;
 	TArray<FString> SelectedPaths;
 
@@ -890,13 +892,14 @@ TArray<FName> FAssetManagerEditorModule::GetContentBrowserSelectedAssetPackages(
 
 	GetAssetDataInPaths(SelectedPaths, SelectedAssets);
 
-	TArray<FName> PackageNames;
+	TSet<FName> PackageNames;
+	PackageNames.Reserve(SelectedAssets.Num());
 	for (const FAssetData& AssetData : SelectedAssets)
 	{
-		OutAssetPackages.AddUnique(AssetData.PackageName);
+		PackageNames.Add(AssetData.PackageName);
 	}
 
-	return OutAssetPackages;
+	return PackageNames.Array();
 }
 
 void FAssetManagerEditorModule::CreateAssetContextMenu(FToolMenuSection& InSection)
