[Reference Viewer] Reference viewer needs better empty state UX

- Adds "No Assets Found" text to path field at top of the Reference Viewer
- Adds centered status widget to display "No Assets Found" text. Same widget could be used for other purposes if necessary.

diff --git a/Engine/Plugins/Editor/AssetManagerEditor/Source/AssetManagerEditor/Private/ReferenceViewer/ReferenceViewerStyle.cpp b/Engine/Plugins/Editor/AssetManagerEditor/Source/AssetManagerEditor/Private/ReferenceViewer/ReferenceViewerStyle.cpp
--- a/Engine/Plugins/Editor/AssetManagerEditor/Source/AssetManagerEditor/Private/ReferenceViewer/ReferenceViewerStyle.cpp
+++ b/Engine/Plugins/Editor/AssetManagerEditor/Source/AssetManagerEditor/Private/ReferenceViewer/ReferenceViewerStyle.cpp
@@ -89,6 +89,20 @@ FReferenceViewerStyle::FReferenceViewerStyle()
 	Set("Icons.Filters", new CORE_IMAGE_BRUSH_SVG("Starship/Common/filter", IconSize));
 	Set("Icons.Duplicate",      new IMAGE_BRUSH_SVG("/GraphNode_Duplicate_8px", IconSize, FStyleColors::White));
 
+	// Centered Status Text
+	{
+		static constexpr FLinearColor CenteredStatusColor = FLinearColor(1, 1, 1, 0.3f);
+		static constexpr float OutlineWidth = 2.0f;
+
+		Set("Graph.CenteredStatusText", FTextBlockStyle(NormalText)
+			.SetFont( DEFAULT_FONT( "BoldCondensed", 16 ) )
+			.SetColorAndOpacity(CenteredStatusColor)
+		);
+
+		// A rounded box brush, showing only its border
+		Set( "Graph.CenteredStatusBrush", new FSlateRoundedBoxBrush(FLinearColor::Transparent, BodyRadius, CenteredStatusColor, OutlineWidth));
+	}
+
 	FSlateStyleRegistry::RegisterSlateStyle(*this);
 }
 
diff --git a/Engine/Plugins/Editor/AssetManagerEditor/Source/AssetManagerEditor/Private/ReferenceViewer/SReferenceViewer.cpp b/Engine/Plugins/Editor/AssetManagerEditor/Source/AssetManagerEditor/Private/ReferenceViewer/SReferenceViewer.cpp
--- a/Engine/Plugins/Editor/AssetManagerEditor/Source/AssetManagerEditor/Private/ReferenceViewer/SReferenceViewer.cpp
+++ b/Engine/Plugins/Editor/AssetManagerEditor/Source/AssetManagerEditor/Private/ReferenceViewer/SReferenceViewer.cpp
@@ -572,6 +572,30 @@ void SReferenceViewer::Construct(const FArguments& InArgs)
 				SNew(STextBlock)
 				.Text(this, &SReferenceViewer::GetStatusText)
 			]
+
+			+SOverlay::Slot()
+			.HAlign(HAlign_Center)
+			.VAlign(VAlign_Center)
+			.Padding(FMargin(0, 0, 0, 16))
+			[
+				SNew(SBox)
+				.MinDesiredWidth(325.0f)
+				.MinDesiredHeight(50.0f)
+				[
+					// Show text within a rounded border
+					SNew(SBorder)
+					.HAlign(HAlign_Center)
+					.VAlign(VAlign_Center)
+					.BorderImage(FReferenceViewerStyle::Get().GetBrush("Graph.CenteredStatusBrush"))
+					.Visibility(this, &SReferenceViewer::GetCenteredStatusVisibility)
+					[
+						SNew(STextBlock)
+						.Justification(ETextJustify::Center)
+						.TextStyle(FReferenceViewerStyle::Get(), "Graph.CenteredStatusText")
+						.Text(this, &SReferenceViewer::GetCenteredStatusText)
+					]
+				]
+			]
 		]
 	];
 
@@ -833,6 +857,10 @@ FText SReferenceViewer::GetAddressBarText() const
 			{
 				return FText::Format(LOCTEXT("AddressBarMultiplePackagesText", "{0} and {1} others"), FText::FromString(CurrentGraphRootPackageNames[0].ToString()), FText::AsNumber(CurrentGraphRootPackageNames.Num()));
 			}
+			else
+			{
+				return LOCTEXT("NoAssetFoundText", "No Assets Found");
+			}
 		}
 		else
 		{
@@ -880,6 +908,26 @@ FText SReferenceViewer::GetStatusText() const
 	return FText();
 }
 
+FText SReferenceViewer::GetCenteredStatusText() const
+{
+	if (GraphObj && GraphObj->Nodes.IsEmpty())
+	{
+		return LOCTEXT("NoAssets", "No Assets Found");
+	}
+
+	return FText();
+}
+
+EVisibility SReferenceViewer::GetCenteredStatusVisibility() const
+{
+	if (GraphObj && GraphObj->Nodes.IsEmpty())
+	{
+		return EVisibility::Visible;
+	}
+
+	return EVisibility::Collapsed;
+}
+
 void SReferenceViewer::OnAddressBarTextCommitted(const FText& NewText, ETextCommit::Type CommitInfo)
 {
 	if (CommitInfo == ETextCommit::OnEnter)
diff --git a/Engine/Plugins/Editor/AssetManagerEditor/Source/AssetManagerEditor/Private/ReferenceViewer/SReferenceViewer.h b/Engine/Plugins/Editor/AssetManagerEditor/Source/AssetManagerEditor/Private/ReferenceViewer/SReferenceViewer.h
--- a/Engine/Plugins/Editor/AssetManagerEditor/Source/AssetManagerEditor/Private/ReferenceViewer/SReferenceViewer.h
+++ b/Engine/Plugins/Editor/AssetManagerEditor/Source/AssetManagerEditor/Private/ReferenceViewer/SReferenceViewer.h
@@ -96,6 +96,12 @@ private:
 
 	/** Gets the text to be displayed for warning/status updates */
 	FText GetStatusText() const;
+	
+	/** Gets the text to be displayed at the center of the graph */
+    FText GetCenteredStatusText() const;
+
+	/** Gets the visibility for the text which can be displayed at the center of the graph */
+	EVisibility GetCenteredStatusVisibility() const;
 
 	/** Called when the path is being edited */
 	void OnAddressBarTextChanged(const FText& NewText);
