Add Revision Control option to copy remote paths
* Fix mixed file/folder selection to properly resolve all items in selection
* Fix revert writable to use filtered file list instead of full file list.
* Expose the number of paths required to trigger user dialog about large operations so it can be customized per operation

#tests Tested copy remote paths action with files, folders and a mix, including local only files

diff --git a/Engine/Plugins/Developer/PerforceSourceControl/Source/PerforceSourceControl/Private/PerforceSourceControlOperations.cpp b/Engine/Plugins/Developer/PerforceSourceControl/Source/PerforceSourceControl/Private/PerforceSourceControlOperations.cpp
--- a/Engine/Plugins/Developer/PerforceSourceControl/Source/PerforceSourceControl/Private/PerforceSourceControlOperations.cpp
+++ b/Engine/Plugins/Developer/PerforceSourceControl/Source/PerforceSourceControl/Private/PerforceSourceControlOperations.cpp
@@ -82,6 +82,7 @@ void IPerforceSourceControlWorker::RegisterWorkers()
 	WorkersMap.Add("DeleteWorkspace", FGetPerforceSourceControlWorker::CreateStatic(&InstantiateWorker<FPerforceDeleteWorkspaceWorker>));
 	WorkersMap.Add("GetFileList", FGetPerforceSourceControlWorker::CreateStatic(&InstantiateWorker<FPerforceGetFileListWorker>));
 	WorkersMap.Add("GetFile", FGetPerforceSourceControlWorker::CreateStatic(&InstantiateWorker<FPerforceGetFileWorker>));
+	WorkersMap.Add("Where", FGetPerforceSourceControlWorker::CreateStatic(&InstantiateWorker<FPerforceWhereWorker>));
 }
 
 TSharedPtr<class IPerforceSourceControlWorker, ESPMode::ThreadSafe> IPerforceSourceControlWorker::CreateWorker(const FName& OperationName, FPerforceSourceControlProvider& SCCProvider)
@@ -3634,4 +3635,57 @@ bool FPerforceGetFileWorker::UpdateStates() const
 	return false;
 }
 
+FName FPerforceWhereWorker::GetName() const
+{
+	return "Where";
+}
+
+bool FPerforceWhereWorker::Execute(FPerforceSourceControlCommand& InCommand)
+{
+	TRACE_CPUPROFILER_EVENT_SCOPE(FPerforceWhereWorker::Execute);
+
+	FScopedPerforceConnection ScopedConnection(InCommand);
+
+	if (!InCommand.IsCanceled() && ScopedConnection.IsValid())
+	{
+		FPerforceConnection& Connection = ScopedConnection.GetConnection();
+		TSharedRef<FWhere, ESPMode::ThreadSafe> Operation = StaticCastSharedRef<FWhere>(InCommand.Operation);
+
+		TArray<FString> Parameters;
+		Parameters.Append(InCommand.Files);
+
+		FP4RecordSet Records;
+		InCommand.bCommandSuccessful = Connection.RunCommand(TEXT("where"), Parameters, Records, InCommand.ResultInfo, FOnIsCancelled::CreateRaw(&InCommand, &FPerforceSourceControlCommand::IsCanceled), InCommand.bConnectionDropped);
+
+		if (InCommand.bCommandSuccessful)
+		{
+			TArray<FWhere::FileInfo> Files;
+			Files.Reserve(Records.Num());
+
+			for (FP4Record& P4Record : Records)
+			{
+				const FString* LocalPath = P4Record.Find("path");
+				const FString* RemotePath = P4Record.Find("depotFile");
+
+				if (ensure(LocalPath) && ensure(RemotePath))
+				{
+					FWhere::FileInfo FileInfo;
+					FileInfo.LocalPath = FPaths::CreateStandardFilename(*LocalPath);
+					FileInfo.RemotePath = *RemotePath;
+					Files.Emplace(MoveTemp(FileInfo));
+				}
+			}
+
+			Operation->SetFiles(MoveTemp(Files));
+		}
+	}
+
+	return InCommand.bCommandSuccessful;
+}
+
+bool FPerforceWhereWorker::UpdateStates() const
+{
+	return false;
+}
+
 #undef LOCTEXT_NAMESPACE
diff --git a/Engine/Plugins/Developer/PerforceSourceControl/Source/PerforceSourceControl/Private/PerforceSourceControlOperations.h b/Engine/Plugins/Developer/PerforceSourceControl/Source/PerforceSourceControl/Private/PerforceSourceControlOperations.h
--- a/Engine/Plugins/Developer/PerforceSourceControl/Source/PerforceSourceControl/Private/PerforceSourceControlOperations.h
+++ b/Engine/Plugins/Developer/PerforceSourceControl/Source/PerforceSourceControl/Private/PerforceSourceControlOperations.h
@@ -522,6 +522,20 @@ public:
 	{}
 	virtual ~FPerforceGetFileWorker() = default;
 
+	// IPerforceSourceControlWorker interface
+	virtual FName GetName() const override;
+	virtual bool Execute(FPerforceSourceControlCommand& InCommand) override;
+	virtual bool UpdateStates() const override;
+};
+
+class FPerforceWhereWorker final : public IPerforceSourceControlWorker
+{
+public:
+	FPerforceWhereWorker(FPerforceSourceControlProvider& InSourceControlProvider)
+		: IPerforceSourceControlWorker(InSourceControlProvider)
+	{}
+	virtual ~FPerforceWhereWorker() = default;
+
 	// IPerforceSourceControlWorker interface
 	virtual FName GetName() const override;
 	virtual bool Execute(FPerforceSourceControlCommand& InCommand) override;
diff --git a/Engine/Plugins/Editor/AssetManagerEditor/AssetManagerEditor.uplugin b/Engine/Plugins/Editor/AssetManagerEditor/AssetManagerEditor.uplugin
--- a/Engine/Plugins/Editor/AssetManagerEditor/AssetManagerEditor.uplugin
+++ b/Engine/Plugins/Editor/AssetManagerEditor/AssetManagerEditor.uplugin
@@ -22,5 +22,11 @@
 			"Type": "UncookedOnly" ,
 			"LoadingPhase" : "PreDefault"
 		}
+	],
+	"Plugins": [
+		{
+			"Name": "ContentBrowserAssetDataSource",
+			"Enabled": true
+		}
 	]
 }
\ No newline at end of file
diff --git a/Engine/Plugins/Editor/AssetManagerEditor/Source/AssetManagerEditor/AssetManagerEditor.Build.cs b/Engine/Plugins/Editor/AssetManagerEditor/Source/AssetManagerEditor/AssetManagerEditor.Build.cs
--- a/Engine/Plugins/Editor/AssetManagerEditor/Source/AssetManagerEditor/AssetManagerEditor.Build.cs
+++ b/Engine/Plugins/Editor/AssetManagerEditor/Source/AssetManagerEditor/AssetManagerEditor.Build.cs
@@ -35,6 +35,7 @@ public class AssetManagerEditor : ModuleRules
 				"CollectionManager",
 				"ContentBrowser",
 				"ContentBrowserData",
+				"ContentBrowserAssetDataSource",
 				"WorkspaceMenuStructure",
 				"AssetDefinition",
 				"AssetTools",
diff --git a/Engine/Plugins/Editor/AssetManagerEditor/Source/AssetManagerEditor/Private/AssetSourceControlContextMenu.cpp b/Engine/Plugins/Editor/AssetManagerEditor/Source/AssetManagerEditor/Private/AssetSourceControlContextMenu.cpp
--- a/Engine/Plugins/Editor/AssetManagerEditor/Source/AssetManagerEditor/Private/AssetSourceControlContextMenu.cpp
+++ b/Engine/Plugins/Editor/AssetManagerEditor/Source/AssetManagerEditor/Private/AssetSourceControlContextMenu.cpp
@@ -9,10 +9,13 @@
 #include "AssetToolsModule.h"
 #include "AssetViewUtils.h"
 #include "Async/Async.h"
+#include "ContentBrowserAssetDataCore.h"
+#include "ContentBrowserAssetDataPayload.h"
 #include "ContentBrowserDataMenuContexts.h"
 #include "ContentBrowserMenuContexts.h"
 #include "FileHelpers.h"
 #include "HAL/IConsoleManager.h"
+#include "HAL/PlatformApplicationMisc.h"
 #include "HAL/PlatformFile.h"
 #include "HAL/PlatformFileManager.h"
 #include "IAssetTools.h"
@@ -40,6 +43,8 @@
 
 namespace UE::AssetSourceControlContextMenu::Private
 {
+	FText LargeOperationTitle = LOCTEXT("SCCLargeOperationWarningTitle", "Continue Operation?");
+
 	DECLARE_DELEGATE_RetVal(bool, FIsAsyncProcessingActive);
 
 	FNewToolMenuCustomWidget MakeCustomWidgetDelegate(const TAttribute<FText>& Label, const TAttribute<FSlateIcon>& Icon, const FIsAsyncProcessingActive& IsAsyncProcessingActive)
@@ -127,7 +132,15 @@ private:
 	bool AddSourceControlMenuOptions(FToolMenuSection& InSection);
 	void FillSourceControlSubMenu(UToolMenu* Menu);
 
-	FExecuteAction ExecutionCheck(FExecuteAction&& InAction) const;
+	// Todo: Make TriggerCheckAtNumActions default value into a user option
+
+	// Warn user if the number of Affected assets is greater than the limit that could affect performance
+	// Use this if the operation works on files only
+	FExecuteAction ExecutionCheck(FExecuteAction&& InAction, int32 TriggerCheckAtNumActions = 10) const;
+	
+	// Warn user if the number of Selected assets is greater than the limit that could affect performance
+	// Use this if the operation can work on files or folders so it only needs the user selected items
+	FExecuteAction ExecutionCheckSelectionOnly(FExecuteAction&& InAction, int32 TriggerCheckAtNumActions = 10) const;
 
 	void ExecuteDiffSelected() const;
 	void ExecuteSCCMerge() const;
@@ -142,6 +155,7 @@ private:
 	void ExecuteSCCRevertWritable() const;
 	void ExecuteSCCSync() const;
 	void ExecuteSCCRefresh() const;
+	void ExecuteSCCCopyPaths() const;
 
 	bool CanExecuteSCCMerge() const;
 	bool CanExecuteSCCCheckOut() const;
@@ -153,6 +167,7 @@ private:
 	bool CanExecuteSCCRevert() const;
 	bool CanExecuteSCCRevertWritable() const;
 	bool CanExecuteSCCSync() const;
+	bool CanExecuteSCCCopyPaths() const;
 	bool CanExecuteSCCDiffAgainstDepot() const;
 	bool CanExecuteDiffSelected() const;
 
@@ -196,11 +211,14 @@ private:
 
 private:
 
-	TArray<FAssetData> SelectedAssets;
+	TArray<FAssetData> SelectedAssets; // Expanded list of assets, no folders
+	TArray<FName> SelectedFolders; // Originating folders that SelectedAssets may have been expanded from
 	TArray<FString> PathsWithUnknownState;
 	TArray<FString> CheckedOutUsers;
 	FText CheckedOutUsersText;
 
+	int32 SelectedFileCount = 0; // Number of files that were user selected, SelectedAssets may be large after recursively expanding SelectedFolders
+
 	TSharedPtr<class ISourceControlOperation, ESPMode::ThreadSafe> SCCOperation;
 	std::atomic<EAsyncState> AsyncState = EAsyncState::None;
 
@@ -218,6 +236,7 @@ private:
 	bool bCanExecuteSCCRevert = false;
 	bool bCanExecuteSCCSync = false;
 	bool bCanExecuteSCCRevertWritable = false;
+	bool bCanExecuteSCCCopyPaths = false;
 };
 
 FAssetSourceControlContextMenu::~FAssetSourceControlContextMenu()
@@ -260,27 +279,45 @@ void FAssetSourceControlContextMenuState::Initialize(FToolMenuSection& InSection
 		}
 	}
 
+	SelectedFileCount = SelectedAssets.Num();
+
 	if (const UContentBrowserDataMenuContext_FolderMenu* ContextObject = InSection.FindContext<UContentBrowserDataMenuContext_FolderMenu>())
 	{
-		TArray<FName> SelectedPaths;
+		// Form a filter from the paths
+		FARFilter Filter;
+		Filter.bRecursivePaths = true;
 
 		for (const FContentBrowserItem& SelectedItem : ContextObject->SelectedItems)
 		{
-			for (const FContentBrowserItemData& SelectedItemData : SelectedItem.GetInternalItems())
+			const FContentBrowserItemData* SelectedItemData = SelectedItem.GetPrimaryInternalItem();
+			if (!SelectedItemData)
 			{
-				SelectedPaths.Add(SelectedItemData.GetInternalPath());
+				continue;
 			}
-		}
 
-		// Load the asset registry module
-		FAssetRegistryModule& AssetRegistryModule = FModuleManager::Get().LoadModuleChecked<FAssetRegistryModule>(TEXT("AssetRegistry"));
+			const UContentBrowserDataSource* DataSource = SelectedItemData->GetOwnerDataSource();
+			if (!DataSource)
+			{
+				continue;
+			}
 
-		// Form a filter from the paths
-		FARFilter Filter;
-		Filter.bRecursivePaths = true;
-		Filter.PackagePaths = MoveTemp(SelectedPaths);
+			for (const FContentBrowserItemData& InternalItem : SelectedItem.GetInternalItems())
+			{
+				if (TSharedPtr<const FContentBrowserAssetFolderItemDataPayload> FolderPayload = ContentBrowserAssetData::GetAssetFolderItemPayload(DataSource, InternalItem))
+				{
+					const FName InternalPath = FolderPayload->GetInternalPath();
+					SelectedFolders.Add(InternalPath);
+					Filter.PackagePaths.Add(InternalPath);
+				}
+				else if (TSharedPtr<const FContentBrowserAssetFileItemDataPayload> AssetPayload = ContentBrowserAssetData::GetAssetFileItemPayload(DataSource, InternalItem))
+				{
+					SelectedAssets.Add(AssetPayload->GetAssetData());
+				}
+			}
+		}
 
 		// Query for a list of assets in the selected paths
+		FAssetRegistryModule& AssetRegistryModule = FModuleManager::Get().LoadModuleChecked<FAssetRegistryModule>(TEXT("AssetRegistry"));
 		AssetRegistryModule.Get().GetAssets(Filter, SelectedAssets);
 
 		bContainsFolders = true;
@@ -387,6 +424,7 @@ void FAssetSourceControlContextMenuState::FillSourceControlSubMenu(UToolMenu* Me
 	const bool bUsesSnapshots = SourceControlProvider.UsesSnapshots();
 	const bool bUsesReadOnly = SourceControlProvider.UsesLocalReadOnlyState();
 	const bool bUsesDiffAgainstDepot = SourceControlProvider.AllowsDiffAgainstDepot();
+	const bool bUsesRemoteFilenames = SourceControlProvider.CanExecuteOperation(ISourceControlOperation::Create<FWhere>());
 
 	if (bUsesFileRevisions)
 	{
@@ -593,17 +631,55 @@ void FAssetSourceControlContextMenuState::FillSourceControlSubMenu(UToolMenu* Me
 			FCanExecuteAction()
 		)
 	);
+
+	if (bUsesRemoteFilenames)
+	{
+		// Where check in p4v is a local conversion so very very fast, bump the number to reduce checks
+		const int32 TriggerCheckAtNumActions = 100;
+
+		Section.AddMenuEntry(
+			"SCCCopyPaths",
+			LOCTEXT("SCCCopyPaths", "Copy Remote Path(s)"),
+			LOCTEXT("SCCCopyPathsTooltip", "Copy the revision control paths of the selected objects to the clipboard."),
+			FSlateIcon(FAppStyle::GetAppStyleSetName(), "GenericCommands.Copy"),
+			FUIAction(
+				ExecutionCheckSelectionOnly(FExecuteAction::CreateSP(this, &FAssetSourceControlContextMenuState::ExecuteSCCCopyPaths), TriggerCheckAtNumActions),
+				FCanExecuteAction::CreateLambda([this]() { return IsActionEnabled(CanExecuteSCCCopyPaths()); })
+			)
+		);
+	}
+}
+
+FExecuteAction FAssetSourceControlContextMenuState::ExecutionCheck(FExecuteAction&& InAction, int32 TriggerCheckAtNumActions) const
+{
+	return FExecuteAction::CreateSPLambda(this, [this, TriggerCheckAtNumActions, Action = MoveTemp(InAction)]()
+		{
+			if (SelectedAssets.Num() > TriggerCheckAtNumActions)
+			{
+				const FText Message = FText::Format(LOCTEXT("SCCLargeOperationWarningMessage", "You are about to perform this operation on a large amount of files ({0}), are you sure you want to continue?\n\nUnreal Editor may become unresponsive."), SelectedAssets.Num());
+				const EAppReturnType::Type Result = FMessageDialog::Open(EAppMsgType::YesNo, Message, UE::AssetSourceControlContextMenu::Private::LargeOperationTitle);
+
+				if (Result != EAppReturnType::Yes)
+				{
+					return;
+				}
+			}
+
+			FScopedSlowTask SlowTask(0, LOCTEXT("SCCOperationSlowTaskLabel", "Performing Revision Control Operation"));
+			SlowTask.MakeDialogDelayed(0.25f);
+			Action.Execute();
+		});
 }
 
-FExecuteAction FAssetSourceControlContextMenuState::ExecutionCheck(FExecuteAction&& InAction) const
+FExecuteAction FAssetSourceControlContextMenuState::ExecutionCheckSelectionOnly(FExecuteAction&& InAction, int32 TriggerCheckAtNumActions) const
 {
-	return FExecuteAction::CreateSPLambda(this, [this, Action = MoveTemp(InAction)]()
+	return FExecuteAction::CreateSPLambda(this, [this, TriggerCheckAtNumActions, Action = MoveTemp(InAction)]()
 		{
-			if (SelectedAssets.Num() > 10) // Todo: Make this into a user option, or different values per operation type
+			const int32 SelectionCount = SelectedFileCount + SelectedFolders.Num();
+			if (SelectionCount > TriggerCheckAtNumActions)
 			{
-				FText Message = FText::Format(LOCTEXT("SCCLargeOperationWarningMessage", "You are about to perform this operation on a large amount of files ({0}), are you sure you want to continue?\n\nUnreal Editor may become unresponsive."), SelectedAssets.Num());
-				FText Title = LOCTEXT("SCCLargeOperationWarningTitle", "Continue Operation?");
-				EAppReturnType::Type Result = FMessageDialog::Open(EAppMsgType::YesNo, Message, Title);
+				const FText Message = FText::Format(LOCTEXT("SCCLargeOperationWarningMessageSelectionOnly", "You are about to perform this operation on a large amount of files/folders ({0}), are you sure you want to continue?\n\nUnreal Editor may become unresponsive."), SelectionCount);
+				const EAppReturnType::Type Result = FMessageDialog::Open(EAppMsgType::YesNo, Message, UE::AssetSourceControlContextMenu::Private::LargeOperationTitle);
 
 				if (Result != EAppReturnType::Yes)
 				{
@@ -920,6 +996,71 @@ void FAssetSourceControlContextMenuState::ExecuteSCCRefresh() const
 	ISourceControlModule::Get().GetProvider().Execute(ISourceControlOperation::Create<FUpdateStatus>(), SourceControlHelpers::PackageFilenames(PackageNames), EConcurrency::Asynchronous);
 }
 
+void FAssetSourceControlContextMenuState::ExecuteSCCCopyPaths() const
+{
+	TArray<FString> PackageNames;
+	GetSelectedPackageNames(PackageNames);
+
+	TArray<FString> PathsToConvert;
+	for (const FName& Folder : SelectedFolders)
+	{
+		PathsToConvert.Add(FPaths::ConvertRelativePathToFull(FPackageName::LongPackageNameToFilename(Folder.ToString())));
+	}
+
+	ISourceControlProvider& SourceControlProvider = ISourceControlModule::Get().GetProvider();
+
+	for (const FString& PackageName : PackageNames)
+	{
+		FSourceControlStatePtr SourceControlState = SourceControlProvider.GetState(SourceControlHelpers::PackageFilename(PackageName), EStateCacheUsage::Use);
+		if (SourceControlState.IsValid() && SourceControlState->IsSourceControlled())
+		{
+			// Don't include any files that are under the folders in the selection
+			if (!SelectedFolders.ContainsByPredicate([&PackageName](const FName& Folder) { return PackageName.StartsWith(Folder.ToString()); }))
+			{
+				PathsToConvert.Add(FPaths::ConvertRelativePathToFull(SourceControlState->GetFilename()));
+			}
+		}
+	}
+
+	TSharedRef<FWhere> Operation = ISourceControlOperation::Create<FWhere>();
+
+	ISourceControlModule::Get().GetProvider().Execute(Operation, PathsToConvert, EConcurrency::Synchronous,
+		FSourceControlOperationComplete::CreateLambda([Operation, NumFolders = SelectedFolders.Num(), &SourceControlProvider](const FSourceControlOperationRef& InOperation, ECommandResult::Type Result)
+			{
+				if (Result == ECommandResult::Succeeded)
+				{
+					TArray<const FString*> RemotePaths;
+					RemotePaths.Reserve(Operation->GetFiles().Num());
+					Algo::Transform(Operation->GetFiles(), RemotePaths, [](const FWhere::FileInfo& Path) { return &Path.RemotePath; });
+
+					TArrayView<const FString*> RemotePathsView = RemotePaths;
+
+					// Sort folders and files separately to keep folders first
+					Algo::Sort(RemotePathsView.Left(NumFolders), [](const FString* Left, const FString* Right) { return Left->Compare(*Right) < 0; });
+					Algo::Sort(RemotePathsView.Right(RemotePaths.Num() - NumFolders), [](const FString* Left, const FString* Right) { return Left->Compare(*Right) < 0; });
+
+					FString ClipboardText;
+
+					for (const FString* Path : RemotePaths)
+					{
+						if (ClipboardText.Len() > 0)
+						{
+							ClipboardText += LINE_TERMINATOR;
+						}
+
+						ClipboardText += *Path;
+					}
+
+					FPlatformApplicationMisc::ClipboardCopy(*ClipboardText);
+				}
+				else
+				{
+					// Todo: Error?
+				}
+			})
+	);
+}
+
 bool FAssetSourceControlContextMenuState::CanExecuteSCCCheckOut() const
 {
 	return bCanExecuteSCCCheckOut;
@@ -970,6 +1111,11 @@ bool FAssetSourceControlContextMenuState::CanExecuteSCCSync() const
 	return bCanExecuteSCCSync;
 }
 
+bool FAssetSourceControlContextMenuState::CanExecuteSCCCopyPaths() const
+{
+	return bCanExecuteSCCCopyPaths;
+}
+
 bool FAssetSourceControlContextMenuState::CanExecuteDiffSelected() const
 {
 	bool bCanDiffSelected = false;
@@ -1180,6 +1326,7 @@ void FAssetSourceControlContextMenuState::TryCacheCanExecuteVars(const TArray<FS
 			else
 			{
 				bCanExecuteSCCHistory = !SourceControlState->IsAdded();
+				bCanExecuteSCCCopyPaths = true;
 			}
 
 			if (SourceControlState->CanRevert())
diff --git a/Engine/Source/Developer/SourceControl/Public/SourceControlOperations.h b/Engine/Source/Developer/SourceControl/Public/SourceControlOperations.h
--- a/Engine/Source/Developer/SourceControl/Public/SourceControlOperations.h
+++ b/Engine/Source/Developer/SourceControl/Public/SourceControlOperations.h
@@ -1412,4 +1412,35 @@ private:
 	FString OutPackageFilename;
 };
 
+/**
+* This operation will get info about the location of a file within a source control provider. The input can be a local or remote path
+* and the result contain both local and remote paths
+*/
+class FWhere : public FSourceControlOperationBase
+{
+public:
+	struct FileInfo
+	{
+		FString LocalPath;
+		FString RemotePath;
+	};
+
+	// ISourceControlOperation interface
+	virtual FName GetName() const override
+	{
+		return "Where";
+	}
+
+	virtual FText GetInProgressString() const override
+	{
+		return LOCTEXT("SourceControl_Where", "Locating file(s) in Revision Control...");
+	}
+
+	void SetFiles(TArray<FileInfo>&& InFiles) { Files = MoveTemp(InFiles); }
+	const TArray<FileInfo>& GetFiles() const { return Files; }
+
+private:
+	TArray<FileInfo> Files;
+};
+
 #undef LOCTEXT_NAMESPACE
