[Reference Viewer] Get content browser tooltip in reference viewer

Based on latest feedback, this CL only adds collection information to the current tooltip.

diff --git a/Engine/Plugins/Editor/AssetManagerEditor/Source/AssetManagerEditor/Private/ReferenceViewer/EdGraphNode_Reference.cpp b/Engine/Plugins/Editor/AssetManagerEditor/Source/AssetManagerEditor/Private/ReferenceViewer/EdGraphNode_Reference.cpp
--- a/Engine/Plugins/Editor/AssetManagerEditor/Source/AssetManagerEditor/Private/ReferenceViewer/EdGraphNode_Reference.cpp
+++ b/Engine/Plugins/Editor/AssetManagerEditor/Source/AssetManagerEditor/Private/ReferenceViewer/EdGraphNode_Reference.cpp
@@ -7,8 +7,10 @@
 #include "Misc/PackageName.h"
 #include "SReferenceViewer.h"
 #include "AssetToolsModule.h"
+#include "CollectionManagerModule.h"
 #include "IAssetTools.h"
 #include "IAssetTypeActions.h"
+#include "ICollectionManager.h"
 #include "ReferenceViewer/EdGraph_ReferenceViewer.h"
 #include "Styling/AppStyle.h"
 
@@ -213,6 +215,7 @@ FText UEdGraphNode_Reference::GetTooltipText() const
 	{
 		TooltipString.Append(GetNodeTitle(ENodeTitleType::FullTitle).ToString() + " collapsed nodes:\n");
 	}
+
 	for (const FAssetIdentifier& AssetId : Identifiers)
 	{
 		if (!TooltipString.IsEmpty())
@@ -230,6 +233,37 @@ FText UEdGraphNode_Reference::GetTooltipText() const
 			break;
 		}
 	}
+
+	if (!IsCollapsed() && !Identifiers.IsEmpty())
+	{
+		const FAssetIdentifier& AssetId = Identifiers[0];
+
+		// Retrieve Collections information
+		TMap<FName, FAssetData> Assets;
+		UE::AssetRegistry::GetAssetForPackages({AssetId.PackageName}, Assets);
+		if (FAssetData* AssetData = Assets.Find(AssetId.PackageName))
+		{
+			ICollectionManager& CollectionManager = FCollectionManagerModule::GetModule().Get();
+
+			TArray<FCollectionNameType> ObjectCollections;
+			CollectionManager.GetCollectionsContainingObject(AssetData->ToSoftObjectPath(), ObjectCollections);
+
+			if (!ObjectCollections.IsEmpty())
+			{
+				TooltipString.Append(TEXT("\n\nCollections:"));
+				for (int32 i = 0; i < ObjectCollections.Num(); i++)
+				{
+					if (i > 0)
+					{
+						TooltipString.Append(TEXT(","));
+					}
+
+					TooltipString.Append(TEXT(" ") + ObjectCollections[i].Name.ToString());
+				}
+			}
+		}
+	}
+
 	return FText::FromString(TooltipString);
 }
 
