[Reference Viewer] Opening a Level inside the ReferenceViewer does not Prompt the Save window if the Level is dirty.

When using the "Edit" option from an asset node context menu, code will now check the selected assets, and filter out the ones which cannot be Edited.
This might also create a "Save Content" window, e.g. if Editing implies loading a different Level and the current one is dirty

diff --git a/Engine/Plugins/Editor/AssetManagerEditor/Source/AssetManagerEditor/Private/AssetManagerEditorModule.cpp b/Engine/Plugins/Editor/AssetManagerEditor/Source/AssetManagerEditor/Private/AssetManagerEditorModule.cpp
--- a/Engine/Plugins/Editor/AssetManagerEditor/Source/AssetManagerEditor/Private/AssetManagerEditorModule.cpp
+++ b/Engine/Plugins/Editor/AssetManagerEditor/Source/AssetManagerEditor/Private/AssetManagerEditorModule.cpp
@@ -45,6 +45,7 @@
 #include "EdGraphSchema_K2.h"
 #include "AssetManagerEditorCommands.h"
 #include "AssetSourceControlContextMenu.h"
+#include "AssetToolsModule.h"
 #include "ReferenceViewer/SReferenceViewer.h"
 #include "ReferenceViewer/SReferenceNode.h"
 #include "ReferenceViewer/EdGraphNode_Reference.h"
@@ -1161,13 +1162,14 @@ void FAssetManagerEditorModule::OnEditAssetIdentifiers(TArray<FAssetIdentifier>
 	TArray<FAssetData> AssetsToLoad;
 	for (FAssetIdentifier AssetIdentifier : AssetIdentifiers)
 	{
+		TArray<FAssetData> AssetDataArray;
 		if (AssetIdentifier.IsPackage())
 		{
 			// Directly a package to load
-			AssetRegistry->GetAssetsByPackageName(AssetIdentifier.PackageName, AssetsToLoad);
+			AssetRegistry->GetAssetsByPackageName(AssetIdentifier.PackageName, AssetDataArray);
 		}
 		else
-		{	
+		{
 			// If it's a primary asset ID, resolve it to a package to load
 			FPrimaryAssetId AssetId = AssetIdentifier.GetPrimaryAssetId();
 			if (AssetId.IsValid())
@@ -1177,7 +1179,29 @@ void FAssetManagerEditorModule::OnEditAssetIdentifiers(TArray<FAssetIdentifier>
 				FAssetData AssetData;
 				if (AssetManager.GetAssetDataForPath(AssetPath, /*out*/ AssetData))
 				{
-					AssetsToLoad.Add(AssetData);
+					AssetDataArray.Add(AssetData);
+				}
+			}
+		}
+
+		FAssetToolsModule& AssetToolsModule = FModuleManager::LoadModuleChecked<FAssetToolsModule>("AssetTools");
+
+		for (const FAssetData& AssetData : AssetDataArray)
+		{
+			TWeakPtr<IAssetTypeActions> AssetTypeActionsWeak = AssetToolsModule.Get().GetAssetTypeActionsForClass(AssetData.GetClass());
+
+			if (AssetTypeActionsWeak.IsValid())
+			{
+				TArray<FAssetData> AssetArray = {AssetData};
+				constexpr bool bIsPreview = false;
+
+				// Check the selected assets, and filter out the ones we cannot Edit.
+				// This might also create a "Save Content" window, e.g. if Editing implies loading a different Level and the current one is dirty 
+				AssetArray = AssetTypeActionsWeak.Pin()->GetValidAssetsForPreviewOrEdit(AssetArray, bIsPreview);
+
+				if (!AssetArray.IsEmpty())
+				{
+					AssetsToLoad.Append(AssetArray);
 				}
 			}
 		}
