[Reference Viewer] Don't count redirectors in depth filter

RecursivelyPopulateNodeInfos now checks whether the current asset is actually a redirector. If that's the case, its Reference Node Info will be marked as such, and the Destination asset will be included in the graph.

RecursivelyCreateNodes will check for redirectors by using the flag set in Reference Node Info, and not increase the depth count in case of a redirector. Graph will explore the Destination asset as well.

diff --git a/Engine/Plugins/Editor/AssetManagerEditor/Source/AssetManagerEditor/Private/ReferenceViewer/EdGraph_ReferenceViewer.cpp b/Engine/Plugins/Editor/AssetManagerEditor/Source/AssetManagerEditor/Private/ReferenceViewer/EdGraph_ReferenceViewer.cpp
--- a/Engine/Plugins/Editor/AssetManagerEditor/Source/AssetManagerEditor/Private/ReferenceViewer/EdGraph_ReferenceViewer.cpp
+++ b/Engine/Plugins/Editor/AssetManagerEditor/Source/AssetManagerEditor/Private/ReferenceViewer/EdGraph_ReferenceViewer.cpp
@@ -23,6 +23,7 @@
 FReferenceNodeInfo::FReferenceNodeInfo(const FAssetIdentifier& InAssetId, bool InbReferencers)
 	: AssetId(InAssetId)
 	, bReferencers(InbReferencers)
+	, bIsRedirector(false)
 	, OverflowCount(0)
 	, bExpandAllChildren(false)
 	, ChildProvisionSize(0)
@@ -34,6 +35,11 @@ bool FReferenceNodeInfo::IsFirstParent(const FAssetIdentifier& InParentId) const
 	return Parents.IsEmpty() || Parents[0] == InParentId;
 }
 
+bool FReferenceNodeInfo::IsRedirector() const
+{
+	return bIsRedirector;
+}
+
 bool FReferenceNodeInfo::IsADuplicate() const
 {
 	return Parents.Num() > 1;
@@ -283,28 +289,34 @@ UEdGraphNode_Reference* UEdGraph_ReferenceViewer::ConstructNodes(const TArray<FA
 			}
 		}
 
-		// Store the AssetData in the NodeInfos
-		TMap<FName, FAssetData> PackagesToAssetDataMap;
-		UE::AssetRegistry::GetAssetForPackages(AllPackageNames.Array(), PackagesToAssetDataMap);
-
-		// Store the AssetData in the NodeInfos and collect Asset Type UClasses to populate the filters
+		// Store the AssetData in the NodeInfos if needed, and collect Asset Type UClasses to populate the filters
 		TSet<FTopLevelAssetPath> AllClasses;
 		for (TPair<FAssetIdentifier, FReferenceNodeInfo>&  InfoPair : NewReferenceNodeInfos)
 		{
-			InfoPair.Value.AssetData = PackagesToAssetDataMap.FindRef(InfoPair.Key.PackageName);
-			if (InfoPair.Value.AssetData.IsValid())
+			// Make sure AssetData is valid
+			if (!InfoPair.Value.AssetData.IsValid())
 			{
-				AllClasses.Add(InfoPair.Value.AssetData.AssetClassPath);
+				const FName& PackageName = InfoPair.Key.PackageName;
+				TMap<FName, FAssetData> PackageToAssetDataMap;
+				UE::AssetRegistry::GetAssetForPackages({PackageName}, PackageToAssetDataMap);
+				InfoPair.Value.AssetData = PackageToAssetDataMap.FindRef(PackageName);
 			}
+
+			AllClasses.Add(InfoPair.Value.AssetData.AssetClassPath);
 		}
 
 		for (TPair<FAssetIdentifier, FReferenceNodeInfo>&  InfoPair : NewDependencyNodeInfos)
 		{
-			InfoPair.Value.AssetData = PackagesToAssetDataMap.FindRef(InfoPair.Key.PackageName);
-			if (InfoPair.Value.AssetData.IsValid())
+			// Make sure AssetData is valid
+			if (!InfoPair.Value.AssetData.IsValid())
 			{
-				AllClasses.Add(InfoPair.Value.AssetData.AssetClassPath);
+				const FName& PackageName = InfoPair.Key.PackageName;
+				TMap<FName, FAssetData> PackageToAssetDataMap;
+				UE::AssetRegistry::GetAssetForPackages({PackageName}, PackageToAssetDataMap);
+				InfoPair.Value.AssetData = PackageToAssetDataMap.FindRef(PackageName);
 			}
+
+			AllClasses.Add(InfoPair.Value.AssetData.AssetClassPath);
 		}
 
 		// Update the cached class types list
@@ -803,11 +815,50 @@ UEdGraph_ReferenceViewer::RecursivelyPopulateNodeInfos(bool bInReferencers, cons
 	check(Identifiers.Num() > 0);
 	int32 ProvisionSize = 0;
 	const FAssetIdentifier& InAssetId = Identifiers[0];
-	if (!ExceedsMaxSearchDepth(InCurrentDepth, InMaxDepth))
+
+	bool bIsRedirector = false;
+
+	// Check if this node is actually a redirector
+	TMap<FName, FAssetData> PackageToAssetDataMap;
+	UE::AssetRegistry::GetAssetForPackages({InAssetId.PackageName}, PackageToAssetDataMap);
+
+	FAssetData* AssetData = PackageToAssetDataMap.Find(InAssetId.PackageName);
+	if (AssetData && AssetData->IsRedirector())
+	{
+		if (UObjectRedirector* Redirector = Cast<UObjectRedirector>(AssetData->GetAsset()))
+		{
+			bIsRedirector = true;
+
+			// We are dealing with a redirector. Let's manually retrieve its Destination Object, and set up its set of nodes explicitly
+			const FName& DestinationPackageName = Redirector->DestinationObject->GetPackage()->GetFName();
+			const FAssetIdentifier DestinationAssetId = FAssetIdentifier::FromString(DestinationPackageName.ToString());
+
+			FReferenceNodeInfo& DestinationReferenceNodeInfo = InNodeInfos.FindOrAdd(DestinationAssetId, FReferenceNodeInfo(DestinationAssetId, bInReferencers));
+
+			// The Destination Node parent is the Redirector one
+			DestinationReferenceNodeInfo.Parents.Emplace(InAssetId);
+
+			// Remove Children from Redirector Node, and just add the Destination Node
+			InNodeInfos[InAssetId].Children.Empty();
+			InNodeInfos[InAssetId].Children.Emplace(DestinationAssetId, EDependencyPinCategory::LinkTypeHard);
+			InNodeInfos[InAssetId].bIsRedirector = true;
+
+			// Populate Info, without increasing current depth - we ignore the Redirector
+			RecursivelyPopulateNodeInfos(bInReferencers, { DestinationAssetId }, InNodeInfos, InCurrentDepth, InMaxDepth);
+		}
+	}
+
+	if (!bIsRedirector && !ExceedsMaxSearchDepth(InCurrentDepth, InMaxDepth))
 	{
 		TMap<FAssetIdentifier, EDependencyPinCategory> ReferenceLinks;
 		GetSortedLinks(Identifiers, bInReferencers, GetReferenceSearchFlags(false), ReferenceLinks);
 
+		// If already available, store Asset Data in Reference Node info
+		if (AssetData)
+		{
+			InNodeInfos[InAssetId].AssetData = *AssetData;
+		}
+
 		InNodeInfos[InAssetId].Children.Reserve(ReferenceLinks.Num());
 		for (const TPair<FAssetIdentifier, EDependencyPinCategory>& Pair : ReferenceLinks)
 		{
@@ -872,6 +923,14 @@ UEdGraphNode_Reference* UEdGraph_ReferenceViewer::RecursivelyCreateNodes(bool bI
 	const FReferenceNodeInfo& NodeInfo = InNodeInfos[InAssetId];
 	int32 NodeProvSize = 1;
 
+	int32 CurrentDepth = InCurrentDepth;
+
+	if (NodeInfo.IsRedirector())
+	{
+		// We don't count depth for redirectors
+		CurrentDepth -= 1;
+	}
+
 	UEdGraphNode_Reference* NewNode = nullptr;
 	if (bIsRoot)
 	{
@@ -889,7 +948,7 @@ UEdGraphNode_Reference* UEdGraph_ReferenceViewer::RecursivelyCreateNodes(bool bI
 
 	bool bIsFirstOccurance = bIsRoot || NodeInfo.IsFirstParent(InParentId);
 	FIntPoint ChildLoc = InNodeLoc;
-	if (!ExceedsMaxSearchDepth(InCurrentDepth, InMaxDepth) && bIsFirstOccurance) // Only expand the first parent
+	if (!ExceedsMaxSearchDepth(CurrentDepth, InMaxDepth) && bIsFirstOccurance) // Only expand the first parent
 	{
 
 		// position the children nodes
@@ -911,27 +970,27 @@ UEdGraphNode_Reference* UEdGraph_ReferenceViewer::RecursivelyCreateNodes(bool bI
 				break;
 			}
 
-		    FAssetIdentifier ChildId = Pair.Key;
-		    int32 ChildProvSize = 0;
-		    if (InNodeInfos[ChildId].IsFirstParent(InAssetId))
-		   	{
-		   		ChildProvSize = InNodeInfos[ChildId].ProvisionSize(InAssetId);
-		   	}
-		   	else if (Settings->GetFindPathEnabled())
-		   	{
-		   		ChildProvSize = 1;
-		   	}
-		   	else if (InNodeInfos[ChildId].PassedFilters && Settings->IsShowDuplicates())
-		   	{
-		   		ChildProvSize = 1;
-		   	}
-
-		    // The provision size will always be at least 1 if it should be shown, factoring in filters, breadth, duplicates, etc.
-		   	if (ChildProvSize > 0)
-		    {
+			FAssetIdentifier ChildId = Pair.Key;
+			int32 ChildProvSize = 0;
+			if (InNodeInfos[ChildId].IsFirstParent(InAssetId))
+			{
+				ChildProvSize = InNodeInfos[ChildId].ProvisionSize(InAssetId);
+			}
+			else if (Settings->GetFindPathEnabled())
+			{
+				ChildProvSize = 1;
+			}
+			else if (InNodeInfos[ChildId].PassedFilters && Settings->IsShowDuplicates())
+			{
+				ChildProvSize = 1;
+			}
+
+			// The provision size will always be at least 1 if it should be shown, factoring in filters, breadth, duplicates, etc.
+			if (ChildProvSize > 0)
+			{
 				ChildLoc.Y += (ChildProvSize - 1) * NodeSizeY * 0.5;
 
-				UEdGraphNode_Reference* ChildNode = RecursivelyCreateNodes(bInReferencers, ChildId, ChildLoc, InAssetId, NewNode, InNodeInfos, InCurrentDepth + 1, InMaxDepth);	
+				UEdGraphNode_Reference* ChildNode = RecursivelyCreateNodes(bInReferencers, ChildId, ChildLoc, InAssetId, NewNode, InNodeInfos, CurrentDepth + 1, InMaxDepth);
 
 				if (bInReferencers)
 				{
@@ -946,7 +1005,7 @@ UEdGraphNode_Reference* UEdGraph_ReferenceViewer::RecursivelyCreateNodes(bool bI
 
 				ChildLoc.Y += NodeSizeY * (ChildProvSize + 1) * 0.5;
 				Breadth ++;
-		    }
+			}
 		}
 
 		// There were more references than allowed to be displayed. Make a collapsed node.
@@ -966,7 +1025,7 @@ UEdGraphNode_Reference* UEdGraph_ReferenceViewer::RecursivelyCreateNodes(bool bI
 					const TPair<FAssetIdentifier, EDependencyPinCategory>& OverflowNodePair = NodeInfo.Children[Breadth];
 
 					const FAssetIdentifier& OverflowNodeAssetId = OverflowNodePair.Key;
-					OverflowNode = RecursivelyCreateNodes(bInReferencers, OverflowNodeAssetId, ChildLoc, NodeInfo.AssetId, NewNode, InNodeInfos, InCurrentDepth + 1, InMaxDepth);
+					OverflowNode = RecursivelyCreateNodes(bInReferencers, OverflowNodeAssetId, ChildLoc, NodeInfo.AssetId, NewNode, InNodeInfos, CurrentDepth + 1, InMaxDepth);
 				}
 			}
 
diff --git a/Engine/Plugins/Editor/AssetManagerEditor/Source/AssetManagerEditor/Public/ReferenceViewer/EdGraph_ReferenceViewer.h b/Engine/Plugins/Editor/AssetManagerEditor/Source/AssetManagerEditor/Public/ReferenceViewer/EdGraph_ReferenceViewer.h
--- a/Engine/Plugins/Editor/AssetManagerEditor/Source/AssetManagerEditor/Public/ReferenceViewer/EdGraph_ReferenceViewer.h
+++ b/Engine/Plugins/Editor/AssetManagerEditor/Source/AssetManagerEditor/Public/ReferenceViewer/EdGraph_ReferenceViewer.h
@@ -35,6 +35,8 @@ struct FReferenceNodeInfo
 	// Which direction.  Referencers are left (other assets that depend on me), Dependencies are right (other assets I depend on)
 	bool bReferencers;
 
+	bool bIsRedirector;
+
 	int32 OverflowCount;
 
 	// Denote when all children have been manually expanded and the breadth limit should be ignored
@@ -43,6 +45,8 @@ struct FReferenceNodeInfo
 	FReferenceNodeInfo(const FAssetIdentifier& InAssetId, bool InbReferencers);
 
 	bool IsFirstParent(const FAssetIdentifier& InParentId) const;
+	
+	bool IsRedirector() const;
 
 	bool IsADuplicate() const;
 
