Use default memory representation so that it matches what the Asset Audit tool shows.

diff --git a/Engine/Plugins/Editor/AssetManagerEditor/Source/AssetManagerEditor/Private/SSizeMap.cpp b/Engine/Plugins/Editor/AssetManagerEditor/Source/AssetManagerEditor/Private/SSizeMap.cpp
--- a/Engine/Plugins/Editor/AssetManagerEditor/Source/AssetManagerEditor/Private/SSizeMap.cpp
+++ b/Engine/Plugins/Editor/AssetManagerEditor/Source/AssetManagerEditor/Private/SSizeMap.cpp
@@ -432,7 +432,7 @@ namespace SizeMapInternals
 		if (SizeInBytes < 1000)
 		{
 			// We ended up with bytes, so show a decimal number
-			SizeText = FText::AsMemory(SizeInBytes, EMemoryUnitStandard::SI);
+			SizeText = FText::AsMemory(SizeInBytes);
 		}
 		else
 		{
@@ -442,7 +442,7 @@ namespace SizeMapInternals
 			NumberFormattingOptions.MinimumFractionalDigits = 0;
 			NumberFormattingOptions.MinimumIntegralDigits = 1;
 
-			SizeText = FText::AsMemory(SizeInBytes, &NumberFormattingOptions, nullptr, EMemoryUnitStandard::SI);
+			SizeText = FText::AsMemory(SizeInBytes, &NumberFormattingOptions);
 		}
 
 		if (!bHasKnownSize)
