[Reference Viewer] Allow Multiple Tabs of Reference Viewer

Adds support for multiple Reference Viewer tabs.
If a Reference Viewer Tab already exists for the specified Asset, it will gain focus and no additional tab will be created.
Code still provides a flag to enable previous single Tab behavior.

diff --git a/Engine/Plugins/Editor/AssetManagerEditor/Source/AssetManagerEditor/Private/AssetManagerEditorModule.cpp b/Engine/Plugins/Editor/AssetManagerEditor/Source/AssetManagerEditor/Private/AssetManagerEditorModule.cpp
--- a/Engine/Plugins/Editor/AssetManagerEditor/Source/AssetManagerEditor/Private/AssetManagerEditorModule.cpp
+++ b/Engine/Plugins/Editor/AssetManagerEditor/Source/AssetManagerEditor/Private/AssetManagerEditorModule.cpp
@@ -51,6 +51,7 @@
 #include "ReferenceViewer/EdGraphNode_Reference.h"
 #include "SSizeMap.h"
 #include "Widgets/Notifications/SNotificationList.h"
+#include "Framework/Docking/TabManager.h"
 #include "Framework/Notifications/NotificationManager.h"
 #include "DesktopPlatformModule.h"
 #include "Misc/FileHelper.h"
@@ -72,6 +73,16 @@
 
 DEFINE_LOG_CATEGORY(LogAssetManagerEditor);
 
+namespace UE::AssetManagerEditor::Private
+{
+	bool bUseMultipleReferenceViewerTabs = true;
+	static FAutoConsoleVariableRef CVarUseMultipleReferenceViewerTabs(
+		TEXT("AssetManagerEditor.OpenReferenceViewerInNewTab"),
+		bUseMultipleReferenceViewerTabs,
+		TEXT("Whether to use multiple Reference Viewer Tabs (one per asset) or not.")
+	);
+}
+
 class FAssetManagerGraphPanelNodeFactory : public FGraphPanelNodeFactory
 {
 	virtual TSharedPtr<SGraphNode> CreateNode(UEdGraphNode* Node) const override
@@ -334,6 +345,12 @@ private:
 
 	static bool GetDependencyTypeArg(const FString& Arg, UE::AssetRegistry::EDependencyQuery& OutRequiredFlags);
 
+	/**
+	 * Tries to open a separate Reference Viewer for the specified Asset Identifier.
+	 * If a Reference Viewer Tab is already showing the Asset, it will be focused
+	 */
+	void OpenAssetReferenceViewerTab(const FAssetIdentifier& InAssetIdentifier, const FReferenceViewerParams& ReferenceViewerParams);
+
 	//Prints all dependency chains from assets in the search path to the target package.
 	void FindReferenceChains(FName TargetPackageName, FName RootSearchPath, UE::AssetRegistry::EDependencyQuery RequiredDependencyFlags);
 
@@ -374,7 +391,7 @@ private:
 
 	TWeakPtr<SDockTab> AssetAuditTab;
 	TWeakPtr<SDockTab> AssetDiskSizeTab;
-	TWeakPtr<SDockTab> ReferenceViewerTab;
+	TMap<FName, TWeakPtr<SDockTab>> ReferenceViewerTabs;
 	TWeakPtr<SDockTab> SizeMapTab;
 	TWeakPtr<SAssetAuditBrowser> AssetAuditUI;
 	TWeakPtr<SAssetTableTreeView> AssetDiskSizeUI;
@@ -405,6 +422,7 @@ private:
 	void OnReloadComplete(EReloadCompleteReason Reason);
 	void OnMarkPackageDirty(UPackage* Pkg, bool bWasDirty);
 	void OnEditAssetIdentifiers(TArray<FAssetIdentifier> AssetIdentifiers);
+	void OnReferenceViewerTabClosed(TSharedRef<SDockTab> InClosedTab);
 
 	TSharedRef<SDockTab> SpawnAssetAuditTab(const FSpawnTabArgs& Args);
 	TSharedRef<SDockTab> SpawnAssetDiskSizeTab(const FSpawnTabArgs& Args);
@@ -537,7 +555,8 @@ void FAssetManagerEditorModule::StartupModule()
 
 		FGlobalTabmanager::Get()->RegisterNomadTabSpawner(ReferenceViewerTabName, FOnSpawnTab::CreateRaw(this, &FAssetManagerEditorModule::SpawnReferenceViewerTab))
 			.SetDisplayName(LOCTEXT("ReferenceViewerTitle", "Reference Viewer"))
-			.SetMenuType(ETabSpawnerMenuType::Hidden);
+			.SetMenuType(ETabSpawnerMenuType::Hidden)
+			.SetIcon(FSlateIcon(FAppStyle::GetAppStyleSetName(), "ContentBrowser.ReferenceViewer"));
 
 		FGlobalTabmanager::Get()->RegisterNomadTabSpawner(SizeMapTabName, FOnSpawnTab::CreateRaw(this, &FAssetManagerEditorModule::SpawnSizeMapTab))
 			.SetDisplayName(LOCTEXT("SizeMapTitle", "Size Map"))
@@ -613,10 +632,18 @@ void FAssetManagerEditorModule::ShutdownModule()
 		{
 			AssetAuditTab.Pin()->RequestCloseTab();
 		}
-		if (ReferenceViewerTab.IsValid())
+
+		if (!ReferenceViewerTabs.IsEmpty())
 		{
-			ReferenceViewerTab.Pin()->RequestCloseTab();
+			for (const TPair<FName, TWeakPtr<SDockTab>>& Pair : ReferenceViewerTabs)
+			{
+				if (const TSharedPtr<SDockTab>& Tab = Pair.Value.Pin())
+				{
+					Tab->RequestCloseTab();
+				}
+			}
 		}
+
 		if (SizeMapTab.IsValid())
 		{
 			SizeMapTab.Pin()->RequestCloseTab();
@@ -687,10 +714,17 @@ TSharedRef<SDockTab> FAssetManagerEditorModule::SpawnAssetDiskSizeTab(const FSpa
 
 TSharedRef<SDockTab> FAssetManagerEditorModule::SpawnReferenceViewerTab(const FSpawnTabArgs& Args)
 {
-	TSharedRef<SDockTab> NewTab = SAssignNew(ReferenceViewerTab, SDockTab)
-		.TabRole(ETabRole::NomadTab);
+	TSharedRef<SDockTab> NewTab = SNew(SDockTab)
+		.TabRole(ETabRole::NomadTab)
+		.OnTabClosed_Raw(this, &FAssetManagerEditorModule::OnReferenceViewerTabClosed)
+		[
+			SAssignNew(ReferenceViewerUI, SReferenceViewer)
+		];
 
-	NewTab->SetContent(SAssignNew(ReferenceViewerUI, SReferenceViewer));
+	if (!ReferenceViewerTabs.Contains(ReferenceViewerTabName))
+	{
+		ReferenceViewerTabs.Emplace(ReferenceViewerTabName, NewTab);
+	}
 
 	return NewTab;
 }
@@ -740,8 +774,104 @@ IAssetManagerEditorModule::FCanOpenReferenceViewerUI& FAssetManagerEditorModule:
 	return CanOpenReferenceViewerUIDelegate;
 }
 
+void FAssetManagerEditorModule::OpenAssetReferenceViewerTab(const FAssetIdentifier& InAssetIdentifier, const FReferenceViewerParams& ReferenceViewerParams)
+{
+	TMap<FName, FAssetData> PackageToAssetDataMap;
+	UE::AssetRegistry::GetAssetForPackages({ InAssetIdentifier.PackageName }, PackageToAssetDataMap);
+
+	FName AssetName;
+
+	if (!PackageToAssetDataMap.IsEmpty())
+	{
+		if (const UObject* Asset = PackageToAssetDataMap[InAssetIdentifier.PackageName].GetAsset())
+		{
+			AssetName = FName(Asset->GetName());
+		}
+	}
+
+	// This happens for C++ classes
+	if (AssetName == NAME_None)
+	{
+		AssetName = InAssetIdentifier.PackageName;
+	}
+
+	// Check if there's already a Reference Viewer Tab for this Asset
+	if (const TWeakPtr<SDockTab>* ExistingTabWeak = ReferenceViewerTabs.Find(AssetName))
+	{
+		if (const TSharedPtr<SDockTab> ExistingTab = ExistingTabWeak->Pin())
+		{
+			FGlobalTabmanager::Get()->DrawAttention(ExistingTab.ToSharedRef());
+		}
+	}
+	else
+	{
+		// No tab available yet for this Asset, let's look for an available Tab, or create a new one.
+		TSharedPtr<SDockTab> ReferenceViewerDefaultTab;
+
+		// Look for the default reference viewer tab, in case one already exists (e.g. from a previous Editor session)
+		if (const TWeakPtr<SDockTab>* ReferenceViewerDefaultTabWeak = ReferenceViewerTabs.Find(ReferenceViewerTabName))
+		{
+			ReferenceViewerDefaultTab = ReferenceViewerDefaultTabWeak->Pin();
+		}
+
+		// Look through Tab Manager as well
+		if (!ReferenceViewerDefaultTab.IsValid())
+		{
+			ReferenceViewerDefaultTab = FGlobalTabmanager::Get()->FindExistingLiveTab(ReferenceViewerTabName);
+		}
+
+		// FindExistingLiveTab fails when looking for tabs located e.g. in details panel area, etc.
+		// So, we make sure we are actually getting a reference to an empty Reference Viewer Tab.
+		// This means its label is still not set to an asset name
+		if (!ReferenceViewerDefaultTab.IsValid())
+		{
+			ReferenceViewerDefaultTab = FGlobalTabmanager::Get()->TryInvokeTab(ReferenceViewerTabName);
+		}
+
+		TSharedPtr<SDockTab> CurrentAssetTab;
+		if (ReferenceViewerDefaultTab && ReferenceViewerDefaultTab->GetTabLabel().EqualTo(FText::FromString("Reference Viewer")))
+		{
+			CurrentAssetTab = ReferenceViewerDefaultTab;
+			FGlobalTabmanager::Get()->DrawAttention(ReferenceViewerDefaultTab.ToSharedRef());
+		}
+
+		if (!CurrentAssetTab)
+		{
+			// There is no existing tab yet for this Asset, we need to create a new one.
+			CurrentAssetTab = SNew(SDockTab)
+				.TabRole(ETabRole::NomadTab)
+				.OnTabClosed_Raw(this, &FAssetManagerEditorModule::OnReferenceViewerTabClosed)
+				[
+					SAssignNew(ReferenceViewerUI, SReferenceViewer)
+				];
+
+			if (CurrentAssetTab)
+			{
+				CurrentAssetTab->SetTabIcon(FSlateIcon(FAppStyle::GetAppStyleSetName(), "ContentBrowser.ReferenceViewer").GetIcon());
+
+				// Try to place the newly created tab next to other Reference Viewer Tabs
+				FGlobalTabmanager::Get()->InsertNewDocumentTab(ReferenceViewerTabName, ReferenceViewerTabName, FTabManager::FLiveTabSearch(ReferenceViewerTabName), CurrentAssetTab.ToSharedRef());
+			}
+		}
+
+		if (CurrentAssetTab)
+		{
+			CurrentAssetTab->SetLabel(FText::FromName(AssetName));
+
+			ReferenceViewerTabs.Emplace(AssetName, CurrentAssetTab);
+			const TSharedPtr<SReferenceViewer>& ReferenceViewer = StaticCastSharedRef<SReferenceViewer>(CurrentAssetTab->GetContent());
+			ReferenceViewer->SetGraphRootIdentifiers({ InAssetIdentifier }, ReferenceViewerParams);
+		}
+	}
+}
+
 void FAssetManagerEditorModule::OpenReferenceViewerUI(const TArray<FAssetIdentifier> SelectedIdentifiers, const FReferenceViewerParams ReferenceViewerParams)
 {
+	// True : will create a dedicated Tab for each specified Asset Identifier. Prevents multiple tabs for the same Asset.
+	// False: legacy behavior, in which a single tab is available, and re-used.
+	bool bOpenInNewTab = true;
+	UE::AssetManagerEditor::Private::CVarUseMultipleReferenceViewerTabs->GetValue(bOpenInNewTab);
+
 	if (!SelectedIdentifiers.IsEmpty())
 	{
 		FText ErrorMessage;
@@ -757,10 +887,20 @@ void FAssetManagerEditorModule::OpenReferenceViewerUI(const TArray<FAssetIdentif
 				}
 			}
 		}
-		else if (TSharedPtr<SDockTab> NewTab = FGlobalTabmanager::Get()->TryInvokeTab(ReferenceViewerTabName))
+		else if (!bOpenInNewTab)
+		{
+			if (const TSharedPtr<SDockTab>& InvokedTab = FGlobalTabmanager::Get()->TryInvokeTab(ReferenceViewerTabName))
+			{
+				TSharedPtr<SReferenceViewer> ReferenceViewer = StaticCastSharedRef<SReferenceViewer>(InvokedTab->GetContent());
+				ReferenceViewer->SetGraphRootIdentifiers(SelectedIdentifiers, ReferenceViewerParams);
+			}
+		}
+		else
 		{
-			TSharedRef<SReferenceViewer> ReferenceViewer = StaticCastSharedRef<SReferenceViewer>(NewTab->GetContent());
-			ReferenceViewer->SetGraphRootIdentifiers(SelectedIdentifiers, ReferenceViewerParams);
+			for (const FAssetIdentifier& AssetIdentifier : SelectedIdentifiers)
+			{
+				OpenAssetReferenceViewerTab(AssetIdentifier, ReferenceViewerParams);
+			}
 		}
 	}
 }
@@ -1225,6 +1365,12 @@ void FAssetManagerEditorModule::OnEditAssetIdentifiers(TArray<FAssetIdentifier>
 	}
 }
 
+void FAssetManagerEditorModule::OnReferenceViewerTabClosed(TSharedRef<SDockTab> InClosedTab)
+{
+	FText TabLabel = InClosedTab->GetTabLabel();
+	ReferenceViewerTabs.Remove(FName(TabLabel.ToString()));
+}
+
 bool FAssetManagerEditorModule::GetManagedPackageListForAssetData(const FAssetData& AssetData, TSet<FName>& ManagedPackageSet)
 {
 	InitializeRegistrySources(true);
