[Reference Viewer] Allow Multiple Tabs of Reference Viewer

Adds support for multiple Reference Viewer tabs.
If a Reference Viewer Tab already exists for the specified Asset, it will gain focus and no additional tab will be created.
Code still provides a flag to enable previous single Tab behavior.

Re-submitting in order to match previous behavior for multiple selection of assets, which will now be opened in a single reference viewer again.

diff --git a/Engine/Plugins/Editor/AssetManagerEditor/Source/AssetManagerEditor/Private/AssetManagerEditorModule.cpp b/Engine/Plugins/Editor/AssetManagerEditor/Source/AssetManagerEditor/Private/AssetManagerEditorModule.cpp
--- a/Engine/Plugins/Editor/AssetManagerEditor/Source/AssetManagerEditor/Private/AssetManagerEditorModule.cpp
+++ b/Engine/Plugins/Editor/AssetManagerEditor/Source/AssetManagerEditor/Private/AssetManagerEditorModule.cpp
@@ -51,6 +51,7 @@
 #include "ReferenceViewer/EdGraphNode_Reference.h"
 #include "SSizeMap.h"
 #include "Widgets/Notifications/SNotificationList.h"
+#include "Framework/Docking/TabManager.h"
 #include "Framework/Notifications/NotificationManager.h"
 #include "DesktopPlatformModule.h"
 #include "Misc/FileHelper.h"
@@ -72,6 +73,18 @@
 
 DEFINE_LOG_CATEGORY(LogAssetManagerEditor);
 
+namespace UE::AssetManagerEditor::Private
+{
+	bool bUseMultipleReferenceViewerTabs = true;
+	static FAutoConsoleVariableRef CVarUseMultipleReferenceViewerTabs(
+		TEXT("AssetManagerEditor.OpenReferenceViewerInNewTab"),
+		bUseMultipleReferenceViewerTabs,
+		TEXT("Whether to use multiple Reference Viewer Tabs (one per asset selection) or not.")
+	);
+
+	static FName ReferenceViewerLabel(TEXT("Reference Viewer"));
+}
+
 class FAssetManagerGraphPanelNodeFactory : public FGraphPanelNodeFactory
 {
 	virtual TSharedPtr<SGraphNode> CreateNode(UEdGraphNode* Node) const override
@@ -334,6 +347,12 @@ private:
 
 	static bool GetDependencyTypeArg(const FString& Arg, UE::AssetRegistry::EDependencyQuery& OutRequiredFlags);
 
+	/**
+	 * Tries to open a separate Reference Viewer for the specified Asset Identifiers.
+	 * If a Reference Viewer Tab is already showing the same Assets selection, it will be focused and no new Tab will be created
+	 */
+	void OpenReferenceViewerTab(const TArray<FAssetIdentifier>& InAssetIdentifiers, const FReferenceViewerParams& ReferenceViewerParams);
+
 	//Prints all dependency chains from assets in the search path to the target package.
 	void FindReferenceChains(FName TargetPackageName, FName RootSearchPath, UE::AssetRegistry::EDependencyQuery RequiredDependencyFlags);
 
@@ -372,9 +391,17 @@ private:
 	FDelegateHandle ReferenceViewerDelegateHandle;
 	FDelegateHandle AssetEditorExtenderDelegateHandle;
 
+	/** Currently opened Reference Viewer Tabs, indexed by their Tab ID*/
+	TMap<FName, TWeakPtr<SDockTab>> ReferenceViewerTabs;
+
+	/** Associates hash created from assets selection to an existing Reference Viewer Tab ID */
+	TMap<uint32, FName> AssetsHashToTabID;
+
+	/** Used to generate unique Tab IDs */
+	uint32 GlobalTabCount = 0;
+
 	TWeakPtr<SDockTab> AssetAuditTab;
 	TWeakPtr<SDockTab> AssetDiskSizeTab;
-	TWeakPtr<SDockTab> ReferenceViewerTab;
 	TWeakPtr<SDockTab> SizeMapTab;
 	TWeakPtr<SAssetAuditBrowser> AssetAuditUI;
 	TWeakPtr<SAssetTableTreeView> AssetDiskSizeUI;
@@ -405,6 +432,7 @@ private:
 	void OnReloadComplete(EReloadCompleteReason Reason);
 	void OnMarkPackageDirty(UPackage* Pkg, bool bWasDirty);
 	void OnEditAssetIdentifiers(TArray<FAssetIdentifier> AssetIdentifiers);
+	void OnReferenceViewerTabClosed(TSharedRef<SDockTab> InClosedTab);
 
 	TSharedRef<SDockTab> SpawnAssetAuditTab(const FSpawnTabArgs& Args);
 	TSharedRef<SDockTab> SpawnAssetDiskSizeTab(const FSpawnTabArgs& Args);
@@ -536,8 +564,9 @@ void FAssetManagerEditorModule::StartupModule()
 		FGlobalTabmanager::Get()->RegisterDefaultTabWindowSize(AssetDiskSize2TabName, FVector2D(1080, 600));
 
 		FGlobalTabmanager::Get()->RegisterNomadTabSpawner(ReferenceViewerTabName, FOnSpawnTab::CreateRaw(this, &FAssetManagerEditorModule::SpawnReferenceViewerTab))
-			.SetDisplayName(LOCTEXT("ReferenceViewerTitle", "Reference Viewer"))
-			.SetMenuType(ETabSpawnerMenuType::Hidden);
+			.SetDisplayName(FText::Format(LOCTEXT("ReferenceViewerTitle", "{0}"), FText::FromName(UE::AssetManagerEditor::Private::ReferenceViewerLabel)))
+			.SetMenuType(ETabSpawnerMenuType::Hidden)
+			.SetIcon(FSlateIcon(FAppStyle::GetAppStyleSetName(), "ContentBrowser.ReferenceViewer"));
 
 		FGlobalTabmanager::Get()->RegisterNomadTabSpawner(SizeMapTabName, FOnSpawnTab::CreateRaw(this, &FAssetManagerEditorModule::SpawnSizeMapTab))
 			.SetDisplayName(LOCTEXT("SizeMapTitle", "Size Map"))
@@ -613,10 +642,21 @@ void FAssetManagerEditorModule::ShutdownModule()
 		{
 			AssetAuditTab.Pin()->RequestCloseTab();
 		}
-		if (ReferenceViewerTab.IsValid())
+
+		if (!ReferenceViewerTabs.IsEmpty())
 		{
-			ReferenceViewerTab.Pin()->RequestCloseTab();
+			for (const TPair<FName, TWeakPtr<SDockTab>>& Pair : ReferenceViewerTabs)
+			{
+				if (const TSharedPtr<SDockTab>& Tab = Pair.Value.Pin())
+				{
+					Tab->RequestCloseTab();
+				}
+			}
+
+			ReferenceViewerTabs.Empty();
+			AssetsHashToTabID.Empty();
 		}
+
 		if (SizeMapTab.IsValid())
 		{
 			SizeMapTab.Pin()->RequestCloseTab();
@@ -687,10 +727,17 @@ TSharedRef<SDockTab> FAssetManagerEditorModule::SpawnAssetDiskSizeTab(const FSpa
 
 TSharedRef<SDockTab> FAssetManagerEditorModule::SpawnReferenceViewerTab(const FSpawnTabArgs& Args)
 {
-	TSharedRef<SDockTab> NewTab = SAssignNew(ReferenceViewerTab, SDockTab)
-		.TabRole(ETabRole::NomadTab);
+	TSharedRef<SDockTab> NewTab = SNew(SDockTab)
+		.OnTabClosed_Raw(this, &FAssetManagerEditorModule::OnReferenceViewerTabClosed)
+		[
+			SAssignNew(ReferenceViewerUI, SReferenceViewer)
+		];
 
-	NewTab->SetContent(SAssignNew(ReferenceViewerUI, SReferenceViewer));
+	if (!ReferenceViewerTabs.Contains(ReferenceViewerTabName))
+	{
+		// This is probably the Reference Viewer Tab created when restoring Editor Layout.
+		ReferenceViewerTabs.Emplace(ReferenceViewerTabName, NewTab);
+	}
 
 	return NewTab;
 }
@@ -740,8 +787,146 @@ IAssetManagerEditorModule::FCanOpenReferenceViewerUI& FAssetManagerEditorModule:
 	return CanOpenReferenceViewerUIDelegate;
 }
 
+void FAssetManagerEditorModule::OpenReferenceViewerTab(const TArray<FAssetIdentifier>& InAssetIdentifiers, const FReferenceViewerParams& ReferenceViewerParams)
+{
+	TArray<FName> PackageNames;
+	for (const FAssetIdentifier& AssetIdentifier : InAssetIdentifiers)
+	{
+		PackageNames.Add(AssetIdentifier.PackageName);
+	}
+
+	TMap<FName, FAssetData> PackageToAssetDataMap;
+	UE::AssetRegistry::GetAssetForPackages(PackageNames, PackageToAssetDataMap);
+
+	FString TabLabel;
+
+	if (!PackageToAssetDataMap.IsEmpty())
+	{
+		TabLabel = PackageToAssetDataMap[PackageNames[0]].AssetName.ToString();
+	}
+
+	// C++ classes will lead to no asset, so we retrieve their package instead
+	// (TODO: this needs to be addressed in order to properly show reference viewer graph for C++ Assets)
+
+	if (!PackageToAssetDataMap.Contains(PackageNames[0]))
+	{
+		if (TabLabel.IsEmpty())
+		{
+			TabLabel = PackageNames[0].ToString();
+		}
+	}
+
+	// Label for multiple assets matches Path field at the top of Reference Viewer graph
+	if (InAssetIdentifiers.Num() > 1)
+	{
+		TabLabel += TEXT(" and ") + FString::FromInt(InAssetIdentifiers.Num() - 1) + TEXT(" others");
+	}
+
+	FName TabID;
+
+	// Create a hash from the concatenation of package names from all the selected assets. This hash is used to match a selection with a Tab ID
+	// This allows to ignore selection order when comparing selections while looking for an existing Reference Viewer for the current selection
+	// We sort package names, so that selection order is not be taken into account
+
+	PackageNames.Sort([](const FName& NameA, const FName& NameB)
+	{ 
+		return NameA.FastLess(NameB);
+	});
+
+	FString StringToHash;
+	for (const FName PackageName : PackageNames)
+	{
+		StringToHash += PackageName.ToString();
+	}
+
+	uint32 SelectionHash = GetTypeHash(FName(StringToHash));
+
+	// Look for possibly existing Tab for the same assets
+	if (const FName* TabIDPtr = AssetsHashToTabID.Find(SelectionHash))
+	{
+		TabID = *TabIDPtr;
+
+		// Check if there's already a Reference Viewer Tab for this Assets selection
+		if (const TWeakPtr<SDockTab>* ExistingTabWeak = ReferenceViewerTabs.Find(TabID))
+		{
+			if (const TSharedPtr<SDockTab>& ExistingTab = ExistingTabWeak->Pin())
+			{
+				FGlobalTabmanager::Get()->DrawAttention(ExistingTab.ToSharedRef());
+				return;
+			}
+		}
+	}
+
+	TSharedPtr<SDockTab> ReferenceViewerDefaultTab;
+
+	// Look for the default reference viewer tab, in case an empty one already exists (e.g. from a previous Editor session)
+	if (const TWeakPtr<SDockTab>* ReferenceViewerDefaultTabWeak = ReferenceViewerTabs.Find(ReferenceViewerTabName))
+	{
+		ReferenceViewerDefaultTab = ReferenceViewerDefaultTabWeak->Pin();
+	}
+
+	// Look through Tab Manager as well
+	if (!ReferenceViewerDefaultTab.IsValid())
+	{
+		ReferenceViewerDefaultTab = FGlobalTabmanager::Get()->FindExistingLiveTab(ReferenceViewerTabName);
+	}
+
+	// FindExistingLiveTab fails when looking for tabs located e.g. in details panel area, etc.
+	// So, we make sure we are actually getting a reference to an empty Reference Viewer Tab.
+	// This means its label is still not set to an asset name
+	if (!ReferenceViewerDefaultTab.IsValid())
+	{
+		ReferenceViewerDefaultTab = FGlobalTabmanager::Get()->TryInvokeTab(ReferenceViewerTabName);
+	}
+
+	TSharedPtr<SDockTab> CurrentAssetTab;
+	if (ReferenceViewerDefaultTab && AssetsHashToTabID.IsEmpty())
+	{
+		TabID = ReferenceViewerTabName;
+		CurrentAssetTab = ReferenceViewerDefaultTab;
+		FGlobalTabmanager::Get()->DrawAttention(ReferenceViewerDefaultTab.ToSharedRef());
+	}
+
+	if (!CurrentAssetTab)
+	{
+		// There is no existing tab yet for this Asset, we need to create a new one.
+		CurrentAssetTab = SNew(SDockTab)
+			.TabRole(ETabRole::NomadTab);
+
+		// Generate Tab ID 
+		TabID = FName(ReferenceViewerTabName.ToString() + "_" + FString::FromInt(++GlobalTabCount));
+
+		// Try to place the newly created tab next to other Reference Viewer Tabs
+		FGlobalTabmanager::Get()->InsertNewDocumentTab(
+			ReferenceViewerTabName,
+			TabID,
+			FTabManager::FLastMajorOrNomadTab(ReferenceViewerTabName),
+			CurrentAssetTab.ToSharedRef());
+	}
+
+	if (CurrentAssetTab)
+	{
+		// Initialize Tab
+		CurrentAssetTab->SetContent(SAssignNew(ReferenceViewerUI, SReferenceViewer));
+		CurrentAssetTab->SetOnTabClosed(SDockTab::FOnTabClosedCallback::CreateRaw(this, &FAssetManagerEditorModule::OnReferenceViewerTabClosed));
+		CurrentAssetTab->SetTabIcon(FSlateIcon(FAppStyle::GetAppStyleSetName(), "ContentBrowser.ReferenceViewer").GetIcon());
+		CurrentAssetTab->SetLabel(FText::FromString(TabLabel));
+
+		ReferenceViewerTabs.Emplace(TabID, CurrentAssetTab);
+		AssetsHashToTabID.Emplace(SelectionHash, TabID);
+
+		const TSharedPtr<SReferenceViewer>& ReferenceViewer = StaticCastSharedRef<SReferenceViewer>(CurrentAssetTab->GetContent());
+		ReferenceViewer->SetGraphRootIdentifiers(InAssetIdentifiers, ReferenceViewerParams);
+	}
+}
+
 void FAssetManagerEditorModule::OpenReferenceViewerUI(const TArray<FAssetIdentifier> SelectedIdentifiers, const FReferenceViewerParams ReferenceViewerParams)
 {
+	// True : will create a dedicated Tab for each specified Asset Identifier. Prevents multiple tabs for the same Asset.
+	// False: legacy behavior, in which a single tab is available, and re-used.
+	bool bOpenInNewTab = true;
+	UE::AssetManagerEditor::Private::CVarUseMultipleReferenceViewerTabs->GetValue(bOpenInNewTab);
+
 	if (!SelectedIdentifiers.IsEmpty())
 	{
 		FText ErrorMessage;
@@ -757,10 +942,18 @@ void FAssetManagerEditorModule::OpenReferenceViewerUI(const TArray<FAssetIdentif
 				}
 			}
 		}
-		else if (TSharedPtr<SDockTab> NewTab = FGlobalTabmanager::Get()->TryInvokeTab(ReferenceViewerTabName))
+		else if (!bOpenInNewTab)
+		{
+			if (const TSharedPtr<SDockTab>& InvokedTab = FGlobalTabmanager::Get()->TryInvokeTab(ReferenceViewerTabName))
+			{
+				TSharedPtr<SReferenceViewer> ReferenceViewer = StaticCastSharedRef<SReferenceViewer>(InvokedTab->GetContent());
+				ReferenceViewer->SetGraphRootIdentifiers(SelectedIdentifiers, ReferenceViewerParams);
+				InvokedTab->SetLabel(FText::FromName(UE::AssetManagerEditor::Private::ReferenceViewerLabel));
+			}
+		}
+		else
 		{
-			TSharedRef<SReferenceViewer> ReferenceViewer = StaticCastSharedRef<SReferenceViewer>(NewTab->GetContent());
-			ReferenceViewer->SetGraphRootIdentifiers(SelectedIdentifiers, ReferenceViewerParams);
+			OpenReferenceViewerTab(SelectedIdentifiers, ReferenceViewerParams);
 		}
 	}
 }
@@ -1225,6 +1418,24 @@ void FAssetManagerEditorModule::OnEditAssetIdentifiers(TArray<FAssetIdentifier>
 	}
 }
 
+void FAssetManagerEditorModule::OnReferenceViewerTabClosed(TSharedRef<SDockTab> InClosedTab)
+{
+	FName TabID = FName(InClosedTab->GetLayoutIdentifier().ToString());
+
+	uint32 PairToRemove = INDEX_NONE;
+	for (const TPair<uint32, FName>& Pair : AssetsHashToTabID)
+	{
+		if (Pair.Value == TabID)
+		{
+			PairToRemove = Pair.Key;
+			break;
+		}
+	}
+
+	AssetsHashToTabID.Remove(PairToRemove);
+	ReferenceViewerTabs.Remove(FName(InClosedTab->GetLayoutIdentifier().ToString()));
+}
+
 bool FAssetManagerEditorModule::GetManagedPackageListForAssetData(const FAssetData& AssetData, TSet<FName>& ManagedPackageSet)
 {
 	InitializeRegistrySources(true);
