[Reference Viewer] Build hash from names without making a large string to track opened Tabs

Unshelved from pending changelist '34117160':
#UE Build hash from names without making a large string

diff --git a/Engine/Plugins/Editor/AssetManagerEditor/Source/AssetManagerEditor/Private/AssetManagerEditorModule.cpp b/Engine/Plugins/Editor/AssetManagerEditor/Source/AssetManagerEditor/Private/AssetManagerEditorModule.cpp
--- a/Engine/Plugins/Editor/AssetManagerEditor/Source/AssetManagerEditor/Private/AssetManagerEditorModule.cpp
+++ b/Engine/Plugins/Editor/AssetManagerEditor/Source/AssetManagerEditor/Private/AssetManagerEditorModule.cpp
@@ -833,14 +833,12 @@ void FAssetManagerEditorModule::OpenReferenceViewerTab(const TArray<FAssetIdenti
 		return NameA.FastLess(NameB);
 	});
 
-	FString StringToHash;
+	uint32 SelectionHash = 0;
 	for (const FName PackageName : PackageNames)
 	{
-		StringToHash += PackageName.ToString();
+		SelectionHash ^= GetTypeHash(PackageName);
 	}
 
-	uint32 SelectionHash = GetTypeHash(FName(StringToHash));
-
 	// Look for possibly existing Tab for the same assets
 	if (const FName* TabIDPtr = AssetsHashToTabID.Find(SelectionHash))
 	{
