[Reference Viewer] Display names of connected properties

Adds functionality to list and display referencing propertiers from a referencer in Reference Viewer graph.
There is a cvar to be set to 1 in order to enable the functionality.
It's likely that at this stage not all cases are properly covered. The functionality is trying to gather properties using assets both as types and values

diff --git a/Engine/Plugins/Editor/AssetManagerEditor/Source/AssetManagerEditor/Private/AssetManagerEditorCommands.cpp b/Engine/Plugins/Editor/AssetManagerEditor/Source/AssetManagerEditor/Private/AssetManagerEditorCommands.cpp
--- a/Engine/Plugins/Editor/AssetManagerEditor/Source/AssetManagerEditor/Private/AssetManagerEditorCommands.cpp
+++ b/Engine/Plugins/Editor/AssetManagerEditor/Source/AssetManagerEditor/Private/AssetManagerEditorCommands.cpp
@@ -28,6 +28,7 @@ void FAssetManagerEditorCommands::RegisterCommands()
 
 	UI_COMMAND(OpenSelectedInAssetEditor, "Edit...", "Opens the selected asset in the relevant editor.", EUserInterfaceActionType::Button, FInputChord(EModifierKey::Control, EKeys::E));
 	UI_COMMAND(ZoomToFit, "Zoom to Fit", "Zoom in and center the view on the selected item", EUserInterfaceActionType::Button, FInputChord(EKeys::F));
+	UI_COMMAND(ResolveReferencingProperties, "Resolve Referencing Properties", "Resolve and display the property (or properties) referencing this asset", EUserInterfaceActionType::Button, FInputChord(EKeys::P));
 
 	UI_COMMAND(IncreaseReferencerSearchDepth, "Increase Referencer Search Depth", "Increase the Referencer Search Depth", EUserInterfaceActionType::Button, FInputChord(EKeys::R));
 	UI_COMMAND(DecreaseReferencerSearchDepth, "Decrease Referencer Search Depth", "Decrease the Referencer Search Depth", EUserInterfaceActionType::Button, FInputChord(EModifierKey::Shift, EKeys::R));
diff --git a/Engine/Plugins/Editor/AssetManagerEditor/Source/AssetManagerEditor/Private/AssetManagerEditorCommands.h b/Engine/Plugins/Editor/AssetManagerEditor/Source/AssetManagerEditor/Private/AssetManagerEditorCommands.h
--- a/Engine/Plugins/Editor/AssetManagerEditor/Source/AssetManagerEditor/Private/AssetManagerEditorCommands.h
+++ b/Engine/Plugins/Editor/AssetManagerEditor/Source/AssetManagerEditor/Private/AssetManagerEditorCommands.h
@@ -130,7 +130,10 @@ public:
 
 	/** Zoom in to fit the selected objects in the window */
 	TSharedPtr<FUICommandInfo> ZoomToFit;
-	
+
+	/** Display the list of properties which are referencing the selected node from its referencing node(s) */
+	TSharedPtr<FUICommandInfo> ResolveReferencingProperties;
+
 	/** Start finding objects */
 	TSharedPtr<FUICommandInfo> Find;
 
diff --git a/Engine/Plugins/Editor/AssetManagerEditor/Source/AssetManagerEditor/Private/AssetManagerEditorModule.cpp b/Engine/Plugins/Editor/AssetManagerEditor/Source/AssetManagerEditor/Private/AssetManagerEditorModule.cpp
--- a/Engine/Plugins/Editor/AssetManagerEditor/Source/AssetManagerEditor/Private/AssetManagerEditorModule.cpp
+++ b/Engine/Plugins/Editor/AssetManagerEditor/Source/AssetManagerEditor/Private/AssetManagerEditorModule.cpp
@@ -48,7 +48,9 @@
 #include "AssetToolsModule.h"
 #include "ReferenceViewer/SReferenceViewer.h"
 #include "ReferenceViewer/SReferenceNode.h"
+#include "ReferenceViewer/SReferencedPropertiesNode.h"
 #include "ReferenceViewer/EdGraphNode_Reference.h"
+#include "ReferenceViewer/EdGraphNode_ReferencedProperties.h"
 #include "SSizeMap.h"
 #include "Widgets/Notifications/SNotificationList.h"
 #include "Framework/Docking/TabManager.h"
@@ -93,6 +95,10 @@ class FAssetManagerGraphPanelNodeFactory : public FGraphPanelNodeFactory
 		{
 			return SNew(SReferenceNode, DependencyNode);
 		}
+		else if (UEdGraphNode_ReferencedProperties* PropertiesNode = Cast<UEdGraphNode_ReferencedProperties>(Node))
+		{
+			return SNew(SReferencedPropertiesNode, PropertiesNode);
+		}
 
 		return nullptr;
 	}
diff --git a/Engine/Plugins/Editor/AssetManagerEditor/Source/AssetManagerEditor/Private/ReferenceViewer/EdGraphNode_ReferencedProperties.cpp b/Engine/Plugins/Editor/AssetManagerEditor/Source/AssetManagerEditor/Private/ReferenceViewer/EdGraphNode_ReferencedProperties.cpp
new file mode 100644
--- /dev/null
+++ b/Engine/Plugins/Editor/AssetManagerEditor/Source/AssetManagerEditor/Private/ReferenceViewer/EdGraphNode_ReferencedProperties.cpp
@@ -0,0 +1,68 @@
+// Copyright Epic Games, Inc. All Rights Reserved.
+
+#include "ReferenceViewer/EdGraphNode_ReferencedProperties.h"
+#include "ReferenceViewer/EdGraphNode_Reference.h"
+#include "ReferenceViewer/ReferenceViewerSettings.h"
+
+#include UE_INLINE_GENERATED_CPP_BY_NAME(EdGraphNode_ReferencedProperties)
+
+FString FReferencingPropertyDescription::GetTypeAsString() const
+{
+	FString TypeString = "";
+
+	switch (Type)
+	{
+	case EAssetReferenceType::Property:
+		TypeString = TEXT("Property Type");
+		break;
+	case EAssetReferenceType::Component:
+		TypeString =  TEXT("Component Type");
+		break;
+	case EAssetReferenceType::Value:
+		TypeString =  TEXT("Property Value");
+		break;
+	case EAssetReferenceType::None:
+		TypeString =  TEXT("");
+		break;
+	default: ;
+	}
+
+	return TypeString;
+}
+
+
+void UEdGraphNode_ReferencedProperties::SetupReferencedPropertiesNode(const TArray<FReferencingPropertyDescription>& InPropertiesDescription
+	, const TObjectPtr<UEdGraphNode_Reference>& InReferencingNode, const TObjectPtr<UEdGraphNode_Reference>& InReferencedNode)
+{
+	ReferencedPropertyDescription = InPropertiesDescription;
+	ReferencingNode = InReferencingNode;
+	ReferencedNode = InReferencedNode;
+
+	RefreshLocation();
+}
+
+void UEdGraphNode_ReferencedProperties::RefreshLocation(const FVector2f& InNodeSize)
+{
+	if (!ReferencedNode || !ReferencingNode)
+	{
+		return;
+	}
+
+	int32 NodeX = (ReferencedNode->NodePosX + ReferencingNode->NodePosX) * 0.5f;
+	int32 NodeY = (ReferencedNode->NodePosY + ReferencingNode->NodePosY) * 0.5f;
+
+	if (const UReferenceViewerSettings* Settings = GetDefault<UReferenceViewerSettings>())
+	{
+		int32 NodeSizeY = Settings->IsCompactMode() ? 100 : 200;
+		NodeSizeY += Settings->IsShowPath() ? 40 : 0;
+		NodeY += NodeSizeY * 0.5f;
+	}
+
+	NodeX -= (InNodeSize.X - 128) * 0.5f;
+	NodeY -= InNodeSize.Y * 0.5f;
+
+	NodePosX = NodeX;
+	NodePosY = NodeY;
+
+	// TODO: This positioning is quite rough, if design stays the same, we could improve it in the future
+}
diff --git a/Engine/Plugins/Editor/AssetManagerEditor/Source/AssetManagerEditor/Private/ReferenceViewer/EdGraph_ReferenceViewer.cpp b/Engine/Plugins/Editor/AssetManagerEditor/Source/AssetManagerEditor/Private/ReferenceViewer/EdGraph_ReferenceViewer.cpp
--- a/Engine/Plugins/Editor/AssetManagerEditor/Source/AssetManagerEditor/Private/ReferenceViewer/EdGraph_ReferenceViewer.cpp
+++ b/Engine/Plugins/Editor/AssetManagerEditor/Source/AssetManagerEditor/Private/ReferenceViewer/EdGraph_ReferenceViewer.cpp
@@ -13,8 +13,14 @@
 #include "ICollectionManager.h"
 #include "CollectionManagerModule.h"
 #include "AssetManagerEditorModule.h"
+#include "SReferencedPropertiesNode.h"
 #include "Engine/AssetManager.h"
+#include "Engine/SCS_Node.h"
+#include "Engine/SimpleConstructionScript.h"
+#include "Engine/UserDefinedStruct.h"
 #include "Interfaces/IPluginManager.h"
+#include "ReferenceViewer/EdGraphNode_ReferencedProperties.h"
+#include "UObject/ReferenceChainSearch.h"
 
 #include UE_INLINE_GENERATED_CPP_BY_NAME(EdGraph_ReferenceViewer)
 
@@ -335,6 +341,267 @@ UEdGraphNode_Reference* UEdGraph_ReferenceViewer::ConstructNodes(const TArray<FA
 	return RefilterGraph();
 }
 
+TArray<FReferencingPropertyDescription> UEdGraph_ReferenceViewer::RetrieveReferencingProperties(UObject* InReferencer, UObject* InReferencedAsset)
+{
+	// This method will check InReferencer for references to InReferencedAsset.
+	// Search includes property types and values.
+	// At this stage, it is possible that some cases won't work well (missing references)
+	// On the other end, some results won't be entirely helpful to the user
+
+	if (!InReferencer || !InReferencedAsset)
+	{
+		return {};
+	}
+
+	TArray<FReferencingPropertyDescription> ReferencingProperties;
+
+	// Registering referencing properties to the output array. Property type defaults to EReferencedPropertyType::Property
+	auto AddReferencingProperty = [&ReferencingProperties](const FString& InPropertyName, const FString& InReferencerName,
+		FReferencingPropertyDescription::EAssetReferenceType InPropertyType = FReferencingPropertyDescription::EAssetReferenceType::Property)
+	{
+		FReferencingPropertyDescription PropertyDescription(InPropertyName, InReferencerName, InPropertyType);
+
+		if (!ReferencingProperties.Contains(PropertyDescription))
+		{
+			ReferencingProperties.AddUnique(PropertyDescription);
+		}
+	};
+
+	// User Defined Struct ("BP Struct")
+	if (UUserDefinedStruct* ReferencerStruct = Cast<UUserDefinedStruct>(InReferencer))
+	{
+		const FProperty* CurrentStructProperty = ReferencerStruct->PropertyLink;
+		while (CurrentStructProperty)
+		{
+			bool bMatchFound = false;
+
+			if (const FObjectPropertyBase* ObjectProperty = CastField<FObjectPropertyBase>(CurrentStructProperty))
+			{
+				if (const TObjectPtr<UClass>& PropertyClass = ObjectProperty->PropertyClass)
+				{
+					bMatchFound = PropertyClass->ClassGeneratedBy == InReferencedAsset;
+				}
+			}
+			else if (const FByteProperty* ByteProperty = CastField<FByteProperty>(CurrentStructProperty))
+			{
+				bMatchFound = ByteProperty->Enum == InReferencedAsset;
+			}
+			else if (const FStructProperty* StructProperty = CastField<FStructProperty>(CurrentStructProperty))
+			{
+				bMatchFound = StructProperty->Struct == InReferencedAsset;
+			}
+
+			if (bMatchFound)
+			{
+				AddReferencingProperty(CurrentStructProperty->GetDisplayNameText().ToString(), InReferencer->GetName());
+			}
+
+			CurrentStructProperty = CurrentStructProperty->PropertyLinkNext;
+		}
+
+		// We are done with this Asset Struct
+		return ReferencingProperties;
+	}
+
+	// In case Referencer is a Blueprint, let's look for BP Actor Components referencing the Referenced Asset
+	if (UBlueprint* ReferencerBlueprint = Cast<UBlueprint>(InReferencer))
+	{
+		if (const TObjectPtr<USimpleConstructionScript>& SimpleConstructionScript = ReferencerBlueprint->SimpleConstructionScript)
+		{
+			const TArray<USCS_Node*>& CDONodes = SimpleConstructionScript->GetAllNodes();
+			for (const USCS_Node* Node : CDONodes)
+			{
+				if (!Node)
+				{
+					continue;
+				}
+
+				UClass* ComponentClass = Node->ComponentClass;
+				if (!ComponentClass)
+				{
+					continue;
+				}
+
+				UObject* GeneratingBlueprintObject = ComponentClass->ClassGeneratedBy;
+				if (!GeneratingBlueprintObject)
+				{
+					continue;
+				}
+
+				if (GeneratingBlueprintObject == InReferencedAsset)
+				{
+					// The blueprint used to generate the current CDO Component Node is the same as the referenced asset: add this to output properties names
+					AddReferencingProperty(*Node->GetVariableName().ToString(), InReferencer->GetName(), FReferencingPropertyDescription::EAssetReferenceType::Component);
+				}
+			}
+		}
+	}
+
+	// This string will be used as support to export properties as text, in case we need it
+	FString PropertyExportString;
+
+	// Going through available fields
+	for (TFieldIterator<FProperty> PropertyIt(InReferencer->GetClass()); PropertyIt; ++PropertyIt)
+	{
+		if (!PropertyIt)
+		{
+			continue;
+		}
+
+		PropertyExportString.Empty();
+
+		// Blueprint Array
+		if (FArrayProperty* ArrayProperty = CastField<FArrayProperty>(*PropertyIt))
+		{
+			FScriptArrayHelper_InContainer ArrayHelper(ArrayProperty, InReferencer);
+			for (int32 ItemIndex = 0; ItemIndex < ArrayHelper.Num(); ItemIndex++)
+			{
+				uint8* ArrayElementMemory = ArrayHelper.GetRawPtr(ItemIndex);
+
+				// Blueprint Property
+				if (ArrayProperty->GetOwnerClass() == UBlueprint::StaticClass())
+				{
+					// We are looking for Blueprint Variables only
+					if (ArrayProperty->GetName() != TEXT("NewVariables"))
+					{
+						continue;
+					}
+
+					const FBPVariableDescription& BPVariableDescription = *reinterpret_cast<const FBPVariableDescription*>(ArrayElementMemory);
+
+					bool bAddProperty = false;
+					UObject* SubCategoryObject = BPVariableDescription.VarType.PinSubCategoryObject.Get();
+					if (SubCategoryObject == InReferencedAsset)
+					{
+						bAddProperty = true;
+					}
+					else if (UClass* BPVariableClass = Cast<UClass>(SubCategoryObject))
+					{
+						const TObjectPtr<UObject>& GeneratingBlueprintObject = BPVariableClass->ClassGeneratedBy;
+						if (GeneratingBlueprintObject == InReferencedAsset)
+						{
+							bAddProperty = true;
+						}
+					}
+					else if (FProperty* InnerProperty = ArrayProperty->Inner)// todo: can we avoid using ExportTextItem_Direct in this case?
+					{
+						InnerProperty->ExportTextItem_Direct(PropertyExportString, ArrayHelper.GetRawPtr(ItemIndex), ArrayHelper.GetRawPtr(ItemIndex), InReferencer, PPF_IncludeTransient);
+						if (!PropertyExportString.IsEmpty() && PropertyExportString.Contains(InReferencedAsset->GetPathName()))
+						{
+							bAddProperty = true;
+						}
+					}
+
+					if (bAddProperty)
+					{
+						AddReferencingProperty(BPVariableDescription.VarName.ToString(), InReferencer->GetName());
+					}
+				}
+				// Other
+				else if (const FProperty* InnerProperty = ArrayProperty->Inner)
+				{
+					if (InnerProperty->IsA<FObjectProperty>())
+					{
+						UObject* Object;
+						InnerProperty->GetValue_InContainer(ArrayElementMemory, &Object);
+						if (Object == InReferencedAsset)
+						{
+							const FString ItemIndexString = "[" + FString::FromInt(ItemIndex) + "]";
+							AddReferencingProperty(ArrayProperty->GetFName().ToString() + ItemIndexString, InReferencer->GetName());
+						}
+					}
+				}
+			}
+		}
+		// Native Array
+		else if (PropertyIt->ArrayDim > 1)
+		{
+			for (int32 ItemIndex = 0; ItemIndex < PropertyIt->ArrayDim; ItemIndex++)
+			{
+				PropertyIt->ExportText_InContainer(ItemIndex, PropertyExportString, InReferencer, InReferencer, InReferencer, PPF_IncludeTransient);
+
+				if (!PropertyExportString.IsEmpty() && PropertyExportString.Contains(InReferencedAsset->GetPathName()))
+				{
+					AddReferencingProperty(PropertyIt->GetFName().ToString(), InReferencer->GetName());
+				}
+			}
+		}
+		else if (PropertyIt->IsA<FObjectProperty>())
+		{
+			UObject* Object;
+			PropertyIt->GetValue_InContainer(InReferencer, &Object);
+			if (Object == InReferencedAsset)
+			{
+				AddReferencingProperty(PropertyIt->GetDisplayNameText().ToString(), InReferencer->GetName(), FReferencingPropertyDescription::EAssetReferenceType::Value);
+			}
+		}
+		// Other property (should handle Struct Property and fields as well)
+		else
+		{
+			PropertyIt->ExportText_InContainer(0, PropertyExportString, InReferencer, InReferencer, InReferencer, PPF_IncludeTransient);
+
+			if (!PropertyExportString.IsEmpty() && PropertyExportString.Contains(InReferencedAsset->GetPathName()))
+			{
+				AddReferencingProperty(PropertyIt->GetDisplayNameText().ToString(), InReferencer->GetName());
+			}
+		}
+	}
+
+	// The code above finds Assets when used as types (e.g. BP Enum, Struct or BPs) but not as values (e.g. a Static Mesh used as variable)
+	// To find those, we check the Reference Chain
+	TSet<FString> FoundProperties;
+	FReferenceChainSearch ReferenceChainSearch(InReferencedAsset, EReferenceChainSearchMode::Direct);
+	for (FReferenceChainSearch::FReferenceChain* ObjectRefChain : ReferenceChainSearch.GetReferenceChains())
+	{
+		for (int32 NodeIndex = 0; NodeIndex < ObjectRefChain->Num(); ++NodeIndex)
+		{
+			const FReferenceChainSearch::FGraphNode* ObjectRefChainLink = ObjectRefChain->GetNode(NodeIndex);
+			UObject* LinkObject = ObjectRefChainLink->ObjectInfo->TryResolveObject();
+
+			if (LinkObject->GetOuter() == InReferencedAsset)
+			{
+				continue;
+			}
+
+			const FReferenceChainSearch::TReferenceInfo<FReferenceChainSearch::FGraphNode>* ReferenceInfo = ObjectRefChain->GetReferenceInfo(NodeIndex);
+			if (!ReferenceInfo)
+			{
+				continue;
+			}
+
+			if (ReferenceInfo->Type != FReferenceChainSearch::EReferenceType::Property)
+			{
+				continue;
+			}
+
+			if (UObject* ReferenceObject = ReferenceInfo->Object->ObjectInfo->TryResolveObject())
+			{
+				if (ReferenceObject == InReferencedAsset)
+				{
+					const FString RefName = ReferenceInfo->ReferencerName.ToString();
+
+					// These should not interest the user
+					if (RefName == "ClassGeneratedBy"
+						|| RefName == "ScriptAndPropertyObjectReferences")
+					{
+						continue;
+					}
+
+					FoundProperties.Add(RefName);
+				}
+			}
+		}
+	}
+
+	// Adding all found values refs to the list of referencing properties
+	for (const FString& PropertyName : FoundProperties)
+	{
+		AddReferencingProperty(PropertyName, InReferencer->GetName(), FReferencingPropertyDescription::EAssetReferenceType::Value);
+	}
+
+	return ReferencingProperties;
+}
+
 UEdGraphNode_Reference* UEdGraph_ReferenceViewer::FindPath(const FAssetIdentifier& RootId, const FAssetIdentifier& TargetId)
 {
 
@@ -1136,6 +1403,12 @@ UEdGraphNode_Reference* UEdGraph_ReferenceViewer::CreateReferenceNode()
 	return Cast<UEdGraphNode_Reference>(CreateNode(UEdGraphNode_Reference::StaticClass(), bSelectNewNode));
 }
 
+UEdGraphNode_ReferencedProperties* UEdGraph_ReferenceViewer::CreateReferencedPropertiesNode()
+{
+	constexpr  bool bSelectNewNode = false;
+	return Cast<UEdGraphNode_ReferencedProperties>(CreateNode(UEdGraphNode_ReferencedProperties::StaticClass(), bSelectNewNode));
+}
+
 void UEdGraph_ReferenceViewer::RemoveAllNodes()
 {
 	TArray<UEdGraphNode*> NodesToRemove = Nodes;
diff --git a/Engine/Plugins/Editor/AssetManagerEditor/Source/AssetManagerEditor/Private/ReferenceViewer/ReferenceViewerSchema.cpp b/Engine/Plugins/Editor/AssetManagerEditor/Source/AssetManagerEditor/Private/ReferenceViewer/ReferenceViewerSchema.cpp
--- a/Engine/Plugins/Editor/AssetManagerEditor/Source/AssetManagerEditor/Private/ReferenceViewer/ReferenceViewerSchema.cpp
+++ b/Engine/Plugins/Editor/AssetManagerEditor/Source/AssetManagerEditor/Private/ReferenceViewer/ReferenceViewerSchema.cpp
@@ -9,6 +9,9 @@
 #include "ConnectionDrawingPolicy.h"
 #include "ToolMenuEntry.h"
 #include "ToolMenuSection.h"
+#include "Materials/MaterialFunctionInterface.h"
+#include "Materials/MaterialInstance.h"
+#include "ReferenceViewer/EdGraphNode_Reference.h"
 
 #include UE_INLINE_GENERATED_CPP_BY_NAME(ReferenceViewerSchema)
 
@@ -153,11 +156,34 @@ void UReferenceViewerSchema::GetContextMenuActions(UToolMenu* Menu, UGraphNodeCo
 	}
 
 	{
+		bool bResolveProperties = false;
+		if (Context)
+		{
+			if (const UEdGraphNode_Reference* ReferenceNode = Cast<UEdGraphNode_Reference>(Context->Node))
+			{
+				bResolveProperties = true;
+				FAssetData AssetData = ReferenceNode->GetAssetData();
+				if (const TObjectPtr<UObject>& CDO = AssetData.GetClass()->ClassDefaultObject)
+				{
+					if (CDO->IsA<UMaterialFunctionInterface>() || CDO->IsA<UMaterialInstance>())
+					{
+						bResolveProperties = false;
+					}
+				}
+			}
+		}
+
 		FToolMenuSection& Section = Menu->AddSection(TEXT("References"), NSLOCTEXT("ReferenceViewerSchema", "ReferencesSectionLabel", "References"));
 		Section.AddMenuEntry(FAssetManagerEditorCommands::Get().CopyReferencedObjects);
 		Section.AddMenuEntry(FAssetManagerEditorCommands::Get().CopyReferencingObjects);
 		Section.AddMenuEntry(FAssetManagerEditorCommands::Get().ShowReferencedObjects);
 		Section.AddMenuEntry(FAssetManagerEditorCommands::Get().ShowReferencingObjects);
+
+		if (bResolveProperties)
+		{
+			Section.AddMenuEntry(FAssetManagerEditorCommands::Get().ResolveReferencingProperties);
+		}
+
 		Section.AddMenuEntry(FAssetManagerEditorCommands::Get().ShowReferenceTree);
 		Section.AddMenuEntry(FAssetManagerEditorCommands::Get().ViewSizeMap);
 
@@ -272,4 +298,3 @@ void UReferenceViewerSchema::GetMakeCollectionWithReferencersOrDependenciesSubMe
 			);
 	}
 }
-
diff --git a/Engine/Plugins/Editor/AssetManagerEditor/Source/AssetManagerEditor/Private/ReferenceViewer/ReferenceViewerStyle.cpp b/Engine/Plugins/Editor/AssetManagerEditor/Source/AssetManagerEditor/Private/ReferenceViewer/ReferenceViewerStyle.cpp
--- a/Engine/Plugins/Editor/AssetManagerEditor/Source/AssetManagerEditor/Private/ReferenceViewer/ReferenceViewerStyle.cpp
+++ b/Engine/Plugins/Editor/AssetManagerEditor/Source/AssetManagerEditor/Private/ReferenceViewer/ReferenceViewerStyle.cpp
@@ -91,8 +91,8 @@ FReferenceViewerStyle::FReferenceViewerStyle()
 
 	// Centered Status Text
 	{
-		static constexpr FLinearColor CenteredStatusColor = FLinearColor(1, 1, 1, 0.3f);
-		static constexpr float OutlineWidth = 2.0f;
+		constexpr FLinearColor CenteredStatusColor = FLinearColor(1, 1, 1, 0.3f);
+		constexpr float OutlineWidth = 2.0f;
 
 		Set("Graph.CenteredStatusText", FTextBlockStyle(NormalText)
 			.SetFont( DEFAULT_FONT( "BoldCondensed", 16 ) )
@@ -103,6 +103,20 @@ FReferenceViewerStyle::FReferenceViewerStyle()
 		Set( "Graph.CenteredStatusBrush", new FSlateRoundedBoxBrush(FLinearColor::Transparent, BodyRadius, CenteredStatusColor, OutlineWidth));
 	}
 
+	// Referenced Properties Border
+	{
+		const FLinearColor OutlineColor = FColor::FromHex("#717171");
+		const FLinearColor FillColor = FColor::FromHex("#282828");
+		constexpr float OutlineWidth = 1.0f;
+
+		// A rounded box brush, showing only its border
+		Set( "Graph.ReferencedPropertiesBrush", new FSlateRoundedBoxBrush(FillColor, BodyRadius, OutlineColor, OutlineWidth));
+
+		Set("Graph.ReferencedPropertiesText", FTextBlockStyle(NormalText)
+			.SetColorAndOpacity(OutlineColor)
+		);
+	}
+
 	FSlateStyleRegistry::RegisterSlateStyle(*this);
 }
 
diff --git a/Engine/Plugins/Editor/AssetManagerEditor/Source/AssetManagerEditor/Private/ReferenceViewer/SReferenceViewer.cpp b/Engine/Plugins/Editor/AssetManagerEditor/Source/AssetManagerEditor/Private/ReferenceViewer/SReferenceViewer.cpp
--- a/Engine/Plugins/Editor/AssetManagerEditor/Source/AssetManagerEditor/Private/ReferenceViewer/SReferenceViewer.cpp
+++ b/Engine/Plugins/Editor/AssetManagerEditor/Source/AssetManagerEditor/Private/ReferenceViewer/SReferenceViewer.cpp
@@ -12,6 +12,7 @@
 #include "ReferenceViewer/HistoryManager.h"
 #include "ReferenceViewerStyle.h"
 #include "ReferenceViewer/EdGraphNode_Reference.h"
+#include "ReferenceViewer/EdGraphNode_ReferencedProperties.h"
 #include "ReferenceViewer/ReferenceViewerSchema.h"
 #include "AssetRegistry/AssetRegistryModule.h"
 #include "ICollectionManager.h"
@@ -26,6 +27,7 @@
 #include "Toolkits/GlobalEditorCommonCommands.h"
 #include "Engine/AssetManager.h"
 #include "ReferenceViewer/SReferenceViewerFilterBar.h"
+#include "ReferenceViewer/SReferencedPropertiesNode.h"
 #include "Widgets/Input/SCheckBox.h"
 #include "Widgets/Input/SComboBox.h"
 #include "Widgets/Input/SComboButton.h"
@@ -1533,6 +1535,11 @@ void SReferenceViewer::RegisterActions()
 		FExecuteAction::CreateSP(this, &SReferenceViewer::ZoomToFit),
 		FCanExecuteAction::CreateSP(this, &SReferenceViewer::CanZoomToFit));
 
+	ReferenceViewerActions->MapAction(
+		FAssetManagerEditorCommands::Get().ResolveReferencingProperties,
+		FExecuteAction::CreateSP(this, &SReferenceViewer::ResolveReferencingProperties),
+		FCanExecuteAction::CreateSP(this, &SReferenceViewer::CanResolveReferencingProperties));
+
 	ReferenceViewerActions->MapAction(
 		FAssetManagerEditorCommands::Get().Find,
 		FExecuteAction::CreateSP(this, &SReferenceViewer::OnFind));
@@ -2289,6 +2296,108 @@ void SReferenceViewer::OnFind()
 	FSlateApplication::Get().SetKeyboardFocus(SearchBox, EFocusCause::SetDirectly);
 }
 
+void SReferenceViewer::ResolveReferencingProperties() const
+{
+	if (!GraphEditorPtr)
+	{
+		return;
+	}
+
+	// Retrieve Object from the specified node. Will load the asset if needed.
+	auto GetObjectFromNode([](const UEdGraphNode_Reference* InNode)
+	{
+		UObject* ReturnObject;
+		if (InNode)
+		{
+			const FAssetData& AssetData = InNode->GetAssetData();
+			if (AssetData.IsAssetLoaded())
+			{
+				ReturnObject = AssetData.GetAsset();
+			}
+			else
+			{
+				FScopedSlowTask SlowTask(0, LOCTEXT("LoadingSelectedObject", "Loading selection..."));
+				SlowTask.MakeDialog();
+				ReturnObject = AssetData.GetAsset();
+			}
+		}
+		else
+		{
+			ReturnObject = nullptr;
+		}
+
+		return ReturnObject;
+	});
+
+	TSet<UObject*> SelectedNodes = GraphEditorPtr->GetSelectedNodes();
+	if (ensure(!SelectedNodes.IsEmpty()))
+	{
+		for (UObject* SelectedNode : SelectedNodes.Array())
+		{
+			UEdGraphNode_Reference* ReferencedNode = Cast<UEdGraphNode_Reference>(SelectedNode);
+			if (!ReferencedNode)
+			{
+				continue;
+			}
+
+			UObject* ReferencedObject = GetObjectFromNode(ReferencedNode);
+			UEdGraphPin* ReferencerPin = ReferencedNode->GetReferencerPin();
+
+			if (!ReferencerPin || !ReferencedObject)
+			{
+				continue;
+			}
+
+			if (ReferencerPin->LinkedTo.IsEmpty())
+			{
+				continue;
+			}
+
+			TArray<FReferencingPropertyDescription> ReferencingProperties;
+			for (UEdGraphPin* ReferencedPin : ReferencerPin->LinkedTo)
+			{
+				if (!ReferencedPin)
+				{
+					continue;
+				}
+
+				UEdGraphNode_Reference* ReferencingNode = Cast<UEdGraphNode_Reference>(ReferencedPin->GetOwningNode());
+				if (!ReferencedNode)
+				{
+					continue;
+				}
+
+				UObject* ReferencingObject = GetObjectFromNode(ReferencingNode);
+				if (!ReferencingObject)
+				{
+					continue;
+				}
+
+				TArray<FReferencingPropertyDescription> ReferencingPropertiesArray = GraphObj->RetrieveReferencingProperties(ReferencingObject, ReferencedObject);
+				if (ReferencingPropertiesArray.IsEmpty())
+				{
+					continue;
+				}
+
+				if (UEdGraphNode_ReferencedProperties* PropertiesNode = GraphObj->CreateReferencedPropertiesNode())
+				{
+					PropertiesNode->SetupReferencedPropertiesNode(ReferencingPropertiesArray, ReferencingNode, ReferencedNode);
+				}
+			}
+		}
+	}
+}
+
+bool SReferenceViewer::CanResolveReferencingProperties() const
+{
+	if (!GraphEditorPtr)
+	{
+		return false;
+	}
+
+	return GraphEditorPtr->GetSelectedNodes().Num() >= 1;
+}
+
 void SReferenceViewer::HandleOnSearchTextChanged(const FText& SearchText)
 {
 	if (GraphObj == nullptr || !GraphEditorPtr.IsValid())
diff --git a/Engine/Plugins/Editor/AssetManagerEditor/Source/AssetManagerEditor/Private/ReferenceViewer/SReferenceViewer.h b/Engine/Plugins/Editor/AssetManagerEditor/Source/AssetManagerEditor/Private/ReferenceViewer/SReferenceViewer.h
--- a/Engine/Plugins/Editor/AssetManagerEditor/Source/AssetManagerEditor/Private/ReferenceViewer/SReferenceViewer.h
+++ b/Engine/Plugins/Editor/AssetManagerEditor/Source/AssetManagerEditor/Private/ReferenceViewer/SReferenceViewer.h
@@ -186,6 +186,8 @@ private:
 	void ZoomToFit();
 	bool CanZoomToFit() const;
 	void OnFind();
+	void ResolveReferencingProperties() const;
+	bool CanResolveReferencingProperties() const;
 
 	/** Find Path */
 	void OnFindPathAssetSelected();
diff --git a/Engine/Plugins/Editor/AssetManagerEditor/Source/AssetManagerEditor/Private/ReferenceViewer/SReferencedPropertiesNode.cpp b/Engine/Plugins/Editor/AssetManagerEditor/Source/AssetManagerEditor/Private/ReferenceViewer/SReferencedPropertiesNode.cpp
new file mode 100644
--- /dev/null
+++ b/Engine/Plugins/Editor/AssetManagerEditor/Source/AssetManagerEditor/Private/ReferenceViewer/SReferencedPropertiesNode.cpp
@@ -0,0 +1,154 @@
+// Copyright Epic Games, Inc. All Rights Reserved.
+
+#include "SReferencedPropertiesNode.h"
+#include "ReferenceViewer/EdGraphNode_ReferencedProperties.h"
+#include "ReferenceViewerStyle.h"
+#include "SlateOptMacros.h"
+#include "Styling/SlateIconFinder.h"
+#include "Widgets/Layout/SBox.h"
+#include "Widgets/SBoxPanel.h"
+#include "Widgets/Images/SImage.h"
+
+/**
+ * Widget representing a referencing property
+ */
+class SReferencedPropertyNode : public SCompoundWidget
+{
+public:
+	SLATE_BEGIN_ARGS(SReferencedPropertiesNode)
+	{
+	}
+
+	SLATE_END_ARGS()
+
+	void Construct(const FArguments& InArgs, const FReferencingPropertyDescription& InReferencingPropertyDescription);
+
+private:
+	FText GetPropertyDisplayName() const;
+	FText GetTooltipText() const;
+	const FSlateBrush* GetIconBrush() const;
+
+	FReferencingPropertyDescription PropertyDescription;
+};
+
+void SReferencedPropertiesNode::Construct(const FArguments& InArgs, UEdGraphNode_ReferencedProperties* InReferencedPropertiesNode)
+{
+	GraphNode = InReferencedPropertiesNode;
+	SetCursor(EMouseCursor::CardinalCross);
+
+	UpdateGraphNode();
+}
+
+void SReferencedPropertiesNode::UpdateGraphNode()
+{
+	// No pins
+	InputPins.Empty();
+	OutputPins.Empty();
+
+	// No side Boxes
+	RightNodeBox.Reset();
+	LeftNodeBox.Reset();
+
+	ContentScale.Bind(this, &SReferencedPropertiesNode::GetContentScale);
+
+	TSharedPtr<SVerticalBox> PropertiesBox;
+
+	GetOrAddSlot(ENodeZone::Center)
+	.HAlign(HAlign_Center)
+	.VAlign(VAlign_Center)
+	[
+		SNew(SBorder)
+		.Padding(FMargin(6.0f, 4.0f, 6.0f, 4.0f))
+		.BorderImage(FReferenceViewerStyle::Get().GetBrush(TEXT("Graph.ReferencedPropertiesBrush")))
+		[
+			SAssignNew(PropertiesBox, SVerticalBox)
+		]
+	];
+
+	if (const UEdGraphNode_ReferencedProperties* ReferencedProperties = Cast<UEdGraphNode_ReferencedProperties>(GraphNode))
+	{
+		for (const FReferencingPropertyDescription& PropertyDescription : ReferencedProperties->GetReferencedPropertiesDescription())
+		{
+			PropertiesBox->AddSlot()
+			[
+				SNew(SReferencedPropertyNode, PropertyDescription)
+			];
+		}
+	}
+}
+
+void SReferencedPropertiesNode::Tick(const FGeometry& InAllottedGeometry, const double InCurrentTime, const float InDeltaTime)
+{
+	SGraphNode::Tick(InAllottedGeometry, InCurrentTime, InDeltaTime);
+
+	if (UEdGraphNode_ReferencedProperties* ReferencedProperties = Cast<UEdGraphNode_ReferencedProperties>(GraphNode))
+	{
+		const FVector2f& Size = InAllottedGeometry.GetLocalSize();
+		ReferencedProperties->RefreshLocation(Size);
+	}
+}
+
+BEGIN_SLATE_FUNCTION_BUILD_OPTIMIZATION
+
+void SReferencedPropertyNode::Construct(const FArguments& InArgs, const FReferencingPropertyDescription& InReferencingPropertyDescription)
+{
+	PropertyDescription = InReferencingPropertyDescription;
+
+	ChildSlot
+	.Padding(FMargin(6.0f, 4.0f, 6.0f, 4.0f))
+	[
+		SNew(SHorizontalBox)
+		+ SHorizontalBox::Slot()
+		.VAlign(VAlign_Center)
+		.Padding(FMargin(0.f, 0.f, 8.f, 0.f))
+		.AutoWidth()
+		[
+			SNew(SImage)
+			.Image(GetIconBrush())
+			.DesiredSizeOverride(FVector2D(16.0f, 16.0f))
+		]
+		+ SHorizontalBox::Slot()
+		.FillWidth(1.0)
+		.VAlign(VAlign_Center)
+		[
+			SNew(SBox)
+			.ToolTipText_Raw(this, &SReferencedPropertyNode::GetTooltipText)
+			[
+				SNew(STextBlock)
+				.TextStyle(FReferenceViewerStyle::Get(), "Graph.ReferencedPropertiesText")
+				.Text(this, &SReferencedPropertyNode::GetPropertyDisplayName)
+			]
+		]
+	];
+}
+
+END_SLATE_FUNCTION_BUILD_OPTIMIZATION
+
+FText SReferencedPropertyNode::GetPropertyDisplayName() const
+{
+	return FText::FromString(PropertyDescription.GetName());
+}
+
+FText SReferencedPropertyNode::GetTooltipText() const
+{
+	return FText::FromString(TEXT("Reference used as ") + PropertyDescription.GetTypeAsString());
+}
+
+const FSlateBrush* SReferencedPropertyNode::GetIconBrush() const
+{
+	const FSlateBrush* ComponentIcon = FSlateIconFinder::FindIconBrushForClass(UObject::StaticClass());
+
+	switch (PropertyDescription.GetType())
+	{
+		case FReferencingPropertyDescription::EAssetReferenceType::Component:
+			ComponentIcon = FSlateIconFinder::FindIconBrushForClass(UActorComponent::StaticClass(), TEXT("SCS.Component"));
+			break;
+		case FReferencingPropertyDescription::EAssetReferenceType::Property:
+		case FReferencingPropertyDescription::EAssetReferenceType::Value:
+		case FReferencingPropertyDescription::EAssetReferenceType::None:
+			break;
+		default: ;
+	}
+
+	return ComponentIcon;
+}
diff --git a/Engine/Plugins/Editor/AssetManagerEditor/Source/AssetManagerEditor/Private/ReferenceViewer/SReferencedPropertiesNode.h b/Engine/Plugins/Editor/AssetManagerEditor/Source/AssetManagerEditor/Private/ReferenceViewer/SReferencedPropertiesNode.h
new file mode 100644
--- /dev/null
+++ b/Engine/Plugins/Editor/AssetManagerEditor/Source/AssetManagerEditor/Private/ReferenceViewer/SReferencedPropertiesNode.h
@@ -0,0 +1,38 @@
+// Copyright Epic Games, Inc. All Rights Reserved.
+
+#pragma once
+
+#include "SGraphNode.h"
+#include "ReferenceViewer/EdGraph_ReferenceViewer.h"
+#include "Widgets/SCompoundWidget.h"
+
+class UEdGraphNode_ReferencedProperties;
+
+/**
+ * Widget displaying the list of properties referencing a specified Asset in the reference Viewer.
+ * It visually represents an UEdGraphNode_ReferencedProperties node
+ */
+class SReferencedPropertiesNode : public SGraphNode
+{
+public:
+	SLATE_BEGIN_ARGS(SReferencedPropertiesNode)
+	{
+	}
+
+	SLATE_END_ARGS()
+
+	void Construct(const FArguments& InArgs, UEdGraphNode_ReferencedProperties* InReferencedPropertiesNode);
+
+	// ~Begin SGraphNode
+	virtual void UpdateGraphNode() override;
+	virtual bool IsNodeEditable() const override { return false; }
+	// ~End SGraphNode
+
+	// ~Begin SNode
+	virtual bool CanBeSelected(const FVector2D& InMousePositionInNode) const override { return false; }
+	// ~End SNode
+
+	// ~Begin SWidget
+	virtual void Tick(const FGeometry& InAllottedGeometry, const double InCurrentTime, const float InDeltaTime) override;
+	// ~End SWidget
+};
diff --git a/Engine/Plugins/Editor/AssetManagerEditor/Source/AssetManagerEditor/Public/ReferenceViewer/EdGraphNode_ReferencedProperties.h b/Engine/Plugins/Editor/AssetManagerEditor/Source/AssetManagerEditor/Public/ReferenceViewer/EdGraphNode_ReferencedProperties.h
new file mode 100644
--- /dev/null
+++ b/Engine/Plugins/Editor/AssetManagerEditor/Source/AssetManagerEditor/Public/ReferenceViewer/EdGraphNode_ReferencedProperties.h
@@ -0,0 +1,96 @@
+// Copyright Epic Games, Inc. All Rights Reserved.
+
+#pragma once
+
+#include "EdGraph/EdGraphNode.h"
+#include "EdGraph_ReferenceViewer.h"
+#include "UObject/ObjectMacros.h"
+#include "EdGraphNode_ReferencedProperties.generated.h"
+
+/**
+ * Describes a graph node property referencing an asset from another node
+ */
+struct FReferencingPropertyDescription
+{
+	enum class EAssetReferenceType : uint8
+	{
+		/** Reference comes from a BP Component Type */
+		Component,
+		/** Reference comes from a BP Variable Type */
+		Property,
+		/** Reference comes from a BP Variable property Value */
+		Value,
+		None
+	};
+
+	FReferencingPropertyDescription(const FString& InName, const FString& InReferencerName, const EAssetReferenceType& InType) :
+		Name(InName),
+		ReferencerName(InReferencerName),
+		Type(InType)
+	{
+	}
+
+	bool operator==(const FReferencingPropertyDescription& InOther) const
+	{
+		return Name == InOther.Name
+			&& ReferencerName == InOther.ReferencerName
+			&& Type == InOther.Type;
+	}
+
+	/** Returns the name of the property */
+	const FString& GetName() const { return Name; }
+
+	/** Returns the name of the property referencer */
+	const FString& GetReferencerName() const { return ReferencerName; }
+
+	/** Returns the property type */
+	EAssetReferenceType GetType() const { return Type; }
+
+	/** Returns the property type as a string (useful e.g. for tooltips) */
+	FString GetTypeAsString() const;
+
+private:
+	friend class SReferencedPropertyNode;
+	FReferencingPropertyDescription() = default;
+
+	FString Name;
+	FString ReferencerName;
+	EAssetReferenceType Type = EAssetReferenceType::None;
+};
+
+
+
+/**
+ * A node to display a list of Node properties which are referencing another Node/Asset
+ */
+UCLASS()
+class ASSETMANAGEREDITOR_API UEdGraphNode_ReferencedProperties : public UEdGraphNode
+{
+	GENERATED_BODY()
+
+public:
+	const TArray<FReferencingPropertyDescription>& GetReferencedPropertiesDescription() const { return ReferencedPropertyDescription; }
+
+	const TObjectPtr<UEdGraphNode_Reference>& GetReferencingNode() const { return ReferencingNode; }
+
+	const TObjectPtr<UEdGraphNode_Reference>& GetReferencedNode() const { return ReferencedNode; }
+
+	/** Initialize this Node */
+	void SetupReferencedPropertiesNode(const TArray<FReferencingPropertyDescription>& InPropertiesDescription
+		, const TObjectPtr<UEdGraphNode_Reference>& InReferencingNode, const TObjectPtr<UEdGraphNode_Reference>& InReferencedNode);
+
+	/**
+	 * Refresh the node location, so it stays mid-way between Referencing and Referenced nodes
+	 * @param InNodeSize: size of the widget representing this node. Useful when centering
+	 */
+	void RefreshLocation(const FVector2f& InNodeSize = FVector2f::Zero());
+
+private:
+	TArray<FReferencingPropertyDescription> ReferencedPropertyDescription;
+
+	UPROPERTY()
+	TObjectPtr<UEdGraphNode_Reference> ReferencingNode;
+
+	UPROPERTY()
+	TObjectPtr<UEdGraphNode_Reference> ReferencedNode;
+};
diff --git a/Engine/Plugins/Editor/AssetManagerEditor/Source/AssetManagerEditor/Public/ReferenceViewer/EdGraph_ReferenceViewer.h b/Engine/Plugins/Editor/AssetManagerEditor/Source/AssetManagerEditor/Public/ReferenceViewer/EdGraph_ReferenceViewer.h
--- a/Engine/Plugins/Editor/AssetManagerEditor/Source/AssetManagerEditor/Public/ReferenceViewer/EdGraph_ReferenceViewer.h
+++ b/Engine/Plugins/Editor/AssetManagerEditor/Source/AssetManagerEditor/Public/ReferenceViewer/EdGraph_ReferenceViewer.h
@@ -13,9 +13,11 @@
 
 class FAssetThumbnailPool;
 class UEdGraphNode_Reference;
+class UEdGraphNode_ReferencedProperties;
 class SReferenceViewer;
 class UReferenceViewerSettings;
 enum class EDependencyPinCategory;
+struct FReferencingPropertyDescription;
 
 /*
 *  Holds asset information for building reference graph
@@ -64,6 +66,7 @@ struct FReferenceNodeInfo
 };
 
 
+
 UCLASS()
 class ASSETMANAGEREDITOR_API UEdGraph_ReferenceViewer : public UEdGraph
 {
@@ -116,6 +119,15 @@ public:
 	bool BreadthLimitExceeded() const { return bBreadthLimitReached; };
 
 private:
+
+	/**
+	 * Retrieves the list of properties/values of a specified Referencer Object which reference the specified Referenced Asset
+	 * @param InReferencer: the Object referencing the Asset
+	 * @param InReferencedAsset: the asset referenced by the specified Object
+	 * @return the list of properties of the input Referencer, referencing the specified Asset.
+	 */
+	TArray<FReferencingPropertyDescription> RetrieveReferencingProperties(UObject* InReferencer, UObject* InReferencedAsset);
+
 	void SetReferenceViewer(TSharedPtr<SReferenceViewer> InViewer);
 	UEdGraphNode_Reference* ConstructNodes(const TArray<FAssetIdentifier>& GraphRootIdentifiers, const FIntPoint& GraphRootOrigin);
 
@@ -126,6 +138,8 @@ private:
 
 	UEdGraphNode_Reference* CreateReferenceNode();
 
+	UEdGraphNode_ReferencedProperties* CreateReferencedPropertiesNode();
+
 	/* Generates a NodeInfo structure then used to generate and layout the graph nodes */
 	void RecursivelyPopulateNodeInfos(bool bReferencers, const TArray<FAssetIdentifier>& Identifiers, TMap<FAssetIdentifier, FReferenceNodeInfo>& NodeInfos, int32 CurrentDepth, int32 MaxDepth);
 
