[Reference Viewer] Remove conflicting key command for showing referencing properties

diff --git a/Engine/Plugins/Editor/AssetManagerEditor/Source/AssetManagerEditor/Private/AssetManagerEditorCommands.cpp b/Engine/Plugins/Editor/AssetManagerEditor/Source/AssetManagerEditor/Private/AssetManagerEditorCommands.cpp
--- a/Engine/Plugins/Editor/AssetManagerEditor/Source/AssetManagerEditor/Private/AssetManagerEditorCommands.cpp
+++ b/Engine/Plugins/Editor/AssetManagerEditor/Source/AssetManagerEditor/Private/AssetManagerEditorCommands.cpp
@@ -28,7 +28,7 @@ void FAssetManagerEditorCommands::RegisterCommands()
 
 	UI_COMMAND(OpenSelectedInAssetEditor, "Edit...", "Opens the selected asset in the relevant editor.", EUserInterfaceActionType::Button, FInputChord(EModifierKey::Control, EKeys::E));
 	UI_COMMAND(ZoomToFit, "Zoom to Fit", "Zoom in and center the view on the selected item", EUserInterfaceActionType::Button, FInputChord(EKeys::F));
-	UI_COMMAND(ResolveReferencingProperties, "Resolve Referencing Properties", "Resolve and display the property (or properties) referencing this asset", EUserInterfaceActionType::Button, FInputChord(EKeys::P));
+	UI_COMMAND(ResolveReferencingProperties, "Resolve Referencing Properties", "Resolve and display the property (or properties) referencing this asset", EUserInterfaceActionType::Button, FInputChord());
 
 	UI_COMMAND(IncreaseReferencerSearchDepth, "Increase Referencer Search Depth", "Increase the Referencer Search Depth", EUserInterfaceActionType::Button, FInputChord(EKeys::R));
 	UI_COMMAND(DecreaseReferencerSearchDepth, "Decrease Referencer Search Depth", "Decrease the Referencer Search Depth", EUserInterfaceActionType::Button, FInputChord(EModifierKey::Shift, EKeys::R));
