Fix crash in reference viewer trying to access CDO of an asset of an unloaded class / a node that's not an asset.

diff --git a/Engine/Plugins/Editor/AssetManagerEditor/Source/AssetManagerEditor/Private/ReferenceViewer/ReferenceViewerSchema.cpp b/Engine/Plugins/Editor/AssetManagerEditor/Source/AssetManagerEditor/Private/ReferenceViewer/ReferenceViewerSchema.cpp
--- a/Engine/Plugins/Editor/AssetManagerEditor/Source/AssetManagerEditor/Private/ReferenceViewer/ReferenceViewerSchema.cpp
+++ b/Engine/Plugins/Editor/AssetManagerEditor/Source/AssetManagerEditor/Private/ReferenceViewer/ReferenceViewerSchema.cpp
@@ -163,12 +163,9 @@ void UReferenceViewerSchema::GetContextMenuActions(UToolMenu* Menu, UGraphNodeCo
 			{
 				bResolveProperties = true;
 				FAssetData AssetData = ReferenceNode->GetAssetData();
-				if (const TObjectPtr<UObject>& CDO = AssetData.GetClass()->ClassDefaultObject)
+				if (AssetData.IsInstanceOf<UMaterialFunctionInterface>() || AssetData.IsInstanceOf<UMaterialInstance>())
 				{
-					if (CDO->IsA<UMaterialFunctionInterface>() || CDO->IsA<UMaterialInstance>())
-					{
-						bResolveProperties = false;
-					}
+					bResolveProperties = false;
 				}
 			}
 		}
