[Reference Viewer] Certain conditions prevent Reference Viewer to open the right tab for the specified asset

diff --git a/Engine/Plugins/Editor/AssetManagerEditor/Source/AssetManagerEditor/Private/AssetManagerEditorModule.cpp b/Engine/Plugins/Editor/AssetManagerEditor/Source/AssetManagerEditor/Private/AssetManagerEditorModule.cpp
--- a/Engine/Plugins/Editor/AssetManagerEditor/Source/AssetManagerEditor/Private/AssetManagerEditorModule.cpp
+++ b/Engine/Plugins/Editor/AssetManagerEditor/Source/AssetManagerEditor/Private/AssetManagerEditorModule.cpp
@@ -426,6 +426,17 @@
 
 	FCanOpenReferenceViewerUI CanOpenReferenceViewerUIDelegate;
 
+	/**
+	 * When currently selected assets are changed for an already opened Reference Viewer Tab,
+	 * this function can be called to properly update Tab Hash and Label
+	 */
+	void UpdateReferenceViewerTabHashAndLabel(const TArray<FAssetIdentifier>& InPreviousSelection, const TArray<FAssetIdentifier>& InCurrentSelection);
+
+	/** Returns a unique Hash for the specified list of Asset Identifiers, and an optional Label to represent them */
+	static uint32 GetHashFromAssetsSelection(const TArray<FAssetIdentifier>& InAssetIdentifiers, FName* OutLabel = nullptr);
+
+	void OnReferenceViewerSelectionChanged(const TArray<FAssetIdentifier>& InPreviousSelection, const TArray<FAssetIdentifier>& InNewSelection);
+
 	void CreateAssetContextMenu(FToolMenuSection& InSection);
 	void OnExtendContentBrowserCommands(TSharedRef<FUICommandList> CommandList, FOnContentBrowserGetSelection GetSelectionDelegate);
 	void OnExtendLevelEditorCommands(TSharedRef<FUICommandList> CommandList);
@@ -655,6 +666,11 @@
 			{
 				if (const TSharedPtr<SDockTab>& Tab = Pair.Value.Pin())
 				{
+					if (const TSharedPtr<SReferenceViewer>& ReferenceViewer = StaticCastSharedRef<SReferenceViewer>(Tab->GetContent()))
+					{
+						ReferenceViewer->OnReferenceViewerSelectionChanged().RemoveAll(this);
+					}
+
 					Tab->RequestCloseTab();
 				}
 			}
@@ -795,56 +811,11 @@
 
 void FAssetManagerEditorModule::OpenReferenceViewerTab(const TArray<FAssetIdentifier>& InAssetIdentifiers, const FReferenceViewerParams& ReferenceViewerParams)
 {
-	TArray<FName> PackageNames;
-	for (const FAssetIdentifier& AssetIdentifier : InAssetIdentifiers)
-	{
-		PackageNames.Add(AssetIdentifier.PackageName);
-	}
+	FName TabLabel;
+	uint32 SelectionHash = GetHashFromAssetsSelection(InAssetIdentifiers, &TabLabel);
 
-	TMap<FName, FAssetData> PackageToAssetDataMap;
-	UE::AssetRegistry::GetAssetForPackages(PackageNames, PackageToAssetDataMap);
-
-	FString TabLabel;
-
-	if (!PackageToAssetDataMap.IsEmpty())
-	{
-		TabLabel = PackageToAssetDataMap[PackageNames[0]].AssetName.ToString();
-	}
-
-	// C++ classes will lead to no asset, so we retrieve their package instead
-	// (TODO: this needs to be addressed in order to properly show reference viewer graph for C++ Assets)
-
-	if (!PackageToAssetDataMap.Contains(PackageNames[0]))
-	{
-		if (TabLabel.IsEmpty())
-		{
-			TabLabel = PackageNames[0].ToString();
-		}
-	}
-
-	// Label for multiple assets matches Path field at the top of Reference Viewer graph
-	if (InAssetIdentifiers.Num() > 1)
-	{
-		TabLabel += TEXT(" and ") + FString::FromInt(InAssetIdentifiers.Num() - 1) + TEXT(" others");
-	}
-
 	FName TabID;
 
-	// Create a hash from the concatenation of package names from all the selected assets. This hash is used to match a selection with a Tab ID
-	// This allows to ignore selection order when comparing selections while looking for an existing Reference Viewer for the current selection
-	// We sort package names, so that selection order is not be taken into account
-
-	PackageNames.Sort([](const FName& NameA, const FName& NameB)
-	{ 
-		return NameA.FastLess(NameB);
-	});
-
-	uint32 SelectionHash = 0;
-	for (const FName PackageName : PackageNames)
-	{
-		SelectionHash ^= GetTypeHash(PackageName);
-	}
-
 	// Look for possibly existing Tab for the same assets
 	if (const FName* TabIDPtr = AssetsHashToTabID.Find(SelectionHash))
 	{
@@ -914,13 +885,15 @@
 		CurrentAssetTab->SetContent(SAssignNew(ReferenceViewerUI, SReferenceViewer));
 		CurrentAssetTab->SetOnTabClosed(SDockTab::FOnTabClosedCallback::CreateRaw(this, &FAssetManagerEditorModule::OnReferenceViewerTabClosed));
 		CurrentAssetTab->SetTabIcon(FSlateIcon(FAppStyle::GetAppStyleSetName(), "ContentBrowser.ReferenceViewer").GetIcon());
-		CurrentAssetTab->SetLabel(FText::FromString(TabLabel));
+		CurrentAssetTab->SetLabel(FText::FromName(TabLabel));
 
 		ReferenceViewerTabs.Emplace(TabID, CurrentAssetTab);
 		AssetsHashToTabID.Emplace(SelectionHash, TabID);
 
 		const TSharedPtr<SReferenceViewer>& ReferenceViewer = StaticCastSharedRef<SReferenceViewer>(CurrentAssetTab->GetContent());
 		ReferenceViewer->SetGraphRootIdentifiers(InAssetIdentifiers, ReferenceViewerParams);
+
+		ReferenceViewer->OnReferenceViewerSelectionChanged().AddRaw(this, &FAssetManagerEditorModule::OnReferenceViewerSelectionChanged);
 	}
 }
 
@@ -1100,6 +1073,114 @@
 	return PackageNames.Array();
 }
 
+void FAssetManagerEditorModule::UpdateReferenceViewerTabHashAndLabel(const TArray<FAssetIdentifier>& InPreviousSelection, const TArray<FAssetIdentifier>& InCurrentSelection)
+{
+	uint32 HashToChange = GetHashFromAssetsSelection(InPreviousSelection);
+
+	FName TabLabel;
+	uint32 NewHash = GetHashFromAssetsSelection(InCurrentSelection, &TabLabel);
+
+	if (!AssetsHashToTabID.Contains(HashToChange))
+	{
+		return;
+	}
+
+	FName TabId;
+	AssetsHashToTabID.RemoveAndCopyValue(HashToChange, TabId);
+	AssetsHashToTabID.Emplace(NewHash, TabId);
+
+	if (const TWeakPtr<SDockTab>* TabWeak = ReferenceViewerTabs.Find(TabId))
+	{
+		if (const TSharedPtr<SDockTab> Tab = TabWeak->Pin())
+		{
+			Tab->SetLabel(FText::FromName(TabLabel));
+		}
+	}
+
+	return;
+}
+
+uint32 FAssetManagerEditorModule::GetHashFromAssetsSelection(const TArray<FAssetIdentifier>& InAssetIdentifiers, FName* OutLabel)
+{
+	TRACE_CPUPROFILER_EVENT_SCOPE(FAssetManagerEditorModule::GetHashFromAssetsSelection);
+
+	TArray<FName> PackageNames;
+	for (const FAssetIdentifier& AssetIdentifier : InAssetIdentifiers)
+	{
+		PackageNames.Add(AssetIdentifier.PackageName);
+	}
+
+	TMap<FName, FAssetData> PackageToAssetDataMap;
+	UE::AssetRegistry::GetAssetForPackages(PackageNames, PackageToAssetDataMap);
+
+	FString TabLabel;
+
+	if (!PackageToAssetDataMap.IsEmpty())
+	{
+		TabLabel = PackageToAssetDataMap[PackageNames[0]].AssetName.ToString();
+	}
+
+	bool bNoAsset = false;
+	// C++ classes will lead to no asset, so we retrieve their package instead
+	// (TODO: this needs to be addressed in order to properly show reference viewer graph for C++ Assets)
+
+	if (!PackageToAssetDataMap.Contains(PackageNames[0]))
+	{
+		bNoAsset = true;
+		if (TabLabel.IsEmpty())
+		{
+			TabLabel = PackageNames[0].ToString();
+
+			if (!InAssetIdentifiers[0].ValueName.IsNone())
+			{
+				TabLabel += TEXT(":") + InAssetIdentifiers[0].ValueName.ToString();
+			}
+		}
+	}
+
+	// Label for multiple assets matches Path field at the top of Reference Viewer graph
+	if (InAssetIdentifiers.Num() > 1)
+	{
+		TabLabel += TEXT(" and ") + FString::FromInt(InAssetIdentifiers.Num() - 1) + TEXT(" others");
+	}
+
+	// Create a hash from the concatenation of package names from all the selected assets. This hash is used to match a selection with a Tab ID
+	// This allows to ignore selection order when comparing selections while looking for an existing Reference Viewer for the current selection
+	// We sort package names, so that selection order is not be taken into account
+
+	PackageNames.Sort([](const FName& NameA, const FName& NameB)
+	{ 
+		return NameA.FastLess(NameB);
+	});
+
+	uint32 SelectionHash = 0;
+	for (const FName PackageName : PackageNames)
+	{
+		// In case there is no asset, we might get the same Hash for potentially different graphs (e.g. GameplayTags),
+		// so let's use the tab label, which in that specific case should include both package name and value
+		if (bNoAsset)
+		{
+			SelectionHash ^= GetTypeHash(TabLabel);
+		}
+		else
+		{
+			SelectionHash ^= GetTypeHash(PackageName);
+		}
+	}
+
+	if (OutLabel)
+	{
+		*OutLabel = FName(TabLabel);
+	}
+
+	return SelectionHash;
+}
+
+void FAssetManagerEditorModule::OnReferenceViewerSelectionChanged(const TArray<FAssetIdentifier>& InPreviousSelection, const TArray<FAssetIdentifier>& InNewSelection)
+{
+	UpdateReferenceViewerTabHashAndLabel(InPreviousSelection, InNewSelection);
+}
+
 void FAssetManagerEditorModule::CreateAssetContextMenu(FToolMenuSection& InSection)
 {
 	InSection.AddMenuEntry(
@@ -1436,6 +1517,11 @@
 		}
 	}
 
+	if (const TSharedPtr<SReferenceViewer>& ReferenceViewer = StaticCastSharedRef<SReferenceViewer>(InClosedTab->GetContent()))
+	{
+		ReferenceViewer->OnReferenceViewerSelectionChanged().RemoveAll(this);
+	}
+
 	AssetsHashToTabID.Remove(PairToRemove);
 	ReferenceViewerTabs.Remove(FName(InClosedTab->GetLayoutIdentifier().ToString()));
 }
diff --git a/Engine/Plugins/Editor/AssetManagerEditor/Source/AssetManagerEditor/Private/ReferenceViewer/SReferenceViewer.cpp b/Engine/Plugins/Editor/AssetManagerEditor/Source/AssetManagerEditor/Private/ReferenceViewer/SReferenceViewer.cpp
--- a/Engine/Plugins/Editor/AssetManagerEditor/Source/AssetManagerEditor/Private/ReferenceViewer/SReferenceViewer.cpp
+++ b/Engine/Plugins/Editor/AssetManagerEditor/Source/AssetManagerEditor/Private/ReferenceViewer/SReferenceViewer.cpp
@@ -706,8 +706,14 @@ void SReferenceViewer::SetCurrentRegistrySource(const FAssetManagerEditorRegistr
 
 void SReferenceViewer::OnNodeDoubleClicked(UEdGraphNode* Node)
 {
+	if (!GraphObj)
+	{
+		return;
+	}
 
-	bool bFoundOverflow = false;
+	const TArray<FAssetIdentifier> CurrentlyVisualizedAssets = GraphObj->GetCurrentGraphRootIdentifiers();
+
+	UEdGraphNode* ParentNode = nullptr;
 	if (UEdGraphNode_Reference* ReferenceNode = Cast<UEdGraphNode_Reference>(Node))
 	{
 		// Overflow nodes have no identifiers
@@ -715,31 +721,26 @@ void SReferenceViewer::OnNodeDoubleClicked(UEdGraphNode* Node)
 		{
 			if (ReferenceNode->GetReferencerPin()->LinkedTo.Num() > 0)
 			{
-				if (UEdGraphNode* ParentNode = ReferenceNode->GetReferencerPin()->LinkedTo[0]->GetOwningNode())
-				{
-					if (UEdGraphNode_Reference* ParentReferenceNode = Cast<UEdGraphNode_Reference>(ParentNode))
-					{
-						FAssetIdentifier ParentID = ParentReferenceNode->GetIdentifier();
-						GraphObj->ExpandNode(false, ParentID);
-						bFoundOverflow = true;
-					}
-				}
+				ParentNode = ReferenceNode->GetReferencerPin()->LinkedTo[0]->GetOwningNode();
 			}
 			else if (ReferenceNode->GetDependencyPin()->LinkedTo.Num() > 0)
 			{
-				if (UEdGraphNode* ParentNode = ReferenceNode->GetDependencyPin()->LinkedTo[0]->GetOwningNode())
-				{
-					if (UEdGraphNode_Reference* ParentReferenceNode = Cast<UEdGraphNode_Reference>(ParentNode))
-					{
-						FAssetIdentifier ParentID = ParentReferenceNode->GetIdentifier();
-						GraphObj->ExpandNode(true, ParentID);
-						bFoundOverflow = true;
-					}
-				}
+				ParentNode = ReferenceNode->GetDependencyPin()->LinkedTo[0]->GetOwningNode();
 			}
 		}
 	}
 
+	bool bFoundOverflow = false;
+	if (ParentNode)
+	{
+		if (UEdGraphNode_Reference* ParentReferenceNode = Cast<UEdGraphNode_Reference>(ParentNode))
+		{
+			FAssetIdentifier ParentID = ParentReferenceNode->GetIdentifier();
+			GraphObj->ExpandNode(true, ParentID);
+			bFoundOverflow = true;
+		}
+	}
+
 	if (!bFoundOverflow)
 	{
 		// turn off the find path tool if the user is wanting to center on another node
@@ -749,6 +750,8 @@ void SReferenceViewer::OnNodeDoubleClicked(UEdGraphNode* Node)
 		Nodes.Add(Node);
 		ReCenterGraphOnNodes( Nodes );
 	}
+
+	OnReferenceViewerSelectionChanged().Broadcast(CurrentlyVisualizedAssets, GraphObj->GetCurrentGraphRootIdentifiers());
 }
 
 void SReferenceViewer::RebuildGraph()
@@ -934,6 +937,13 @@ void SReferenceViewer::OnAddressBarTextCommitted(const FText& NewText, ETextComm
 {
 	if (CommitInfo == ETextCommit::OnEnter)
 	{
+		if (!GraphObj)
+		{
+			return;
+		}
+
+		const TArray<FAssetIdentifier> CurrentlyVisualizedAssets = GraphObj->GetCurrentGraphRootIdentifiers();
+
 		TArray<FAssetIdentifier> NewPaths;
 		FAssetIdentifier NewPath = FAssetIdentifier::FromString(NewText.ToString());
 
@@ -948,6 +958,8 @@ void SReferenceViewer::OnAddressBarTextCommitted(const FText& NewText, ETextComm
 		}
 
 		SetGraphRootIdentifiers(NewPaths);
+
+		OnReferenceViewerSelectionChanged().Broadcast(CurrentlyVisualizedAssets, GraphObj->GetCurrentGraphRootIdentifiers());
 	}
 }
 
diff --git a/Engine/Plugins/Editor/AssetManagerEditor/Source/AssetManagerEditor/Private/ReferenceViewer/SReferenceViewer.h b/Engine/Plugins/Editor/AssetManagerEditor/Source/AssetManagerEditor/Private/ReferenceViewer/SReferenceViewer.h
--- a/Engine/Plugins/Editor/AssetManagerEditor/Source/AssetManagerEditor/Private/ReferenceViewer/SReferenceViewer.h
+++ b/Engine/Plugins/Editor/AssetManagerEditor/Source/AssetManagerEditor/Private/ReferenceViewer/SReferenceViewer.h
@@ -19,6 +19,8 @@ template <typename OptionType> class SComboBox;
 
 class UEdGraph;
 
+DECLARE_MULTICAST_DELEGATE_TwoParams(FOnReferenceViewerSelectionChanged, const TArray<FAssetIdentifier>&, const TArray<FAssetIdentifier>&)
+
 /**
  * 
  */
@@ -53,6 +55,8 @@ public:
 	virtual FReply OnKeyDown( const FGeometry& MyGeometry, const FKeyEvent& InKeyEvent ) override;
 	virtual void Tick( const FGeometry& AllottedGeometry, const double InCurrentTime, const float InDeltaTime ) override;
 
+	FOnReferenceViewerSelectionChanged& OnReferenceViewerSelectionChanged() { return OnReferenceViewerSelectionChangedDelegate; }
+
 private:
 
 	/** Call after a structural change is made that causes the graph to be recreated */
@@ -287,6 +291,9 @@ private:
 
 	/** Handle to know if dirty */
 	FDelegateHandle AssetRefreshHandle;
+
+	/** Called when expanding a node, or manually updating the asset path */
+	FOnReferenceViewerSelectionChanged OnReferenceViewerSelectionChangedDelegate;
 };
 
 enum class EDependencyPinCategory
