#UE Make GetStyleSetName return a unique value so it has its own namespace to register the slate style

diff --git a/Engine/Plugins/Editor/AssetManagerEditor/Source/AssetManagerEditor/Private/Insights/Common/InsightsStyle.cpp b/Engine/Plugins/Editor/AssetManagerEditor/Source/AssetManagerEditor/Private/Insights/Common/InsightsStyle.cpp
--- a/Engine/Plugins/Editor/AssetManagerEditor/Source/AssetManagerEditor/Private/Insights/Common/InsightsStyle.cpp
+++ b/Engine/Plugins/Editor/AssetManagerEditor/Source/AssetManagerEditor/Private/Insights/Common/InsightsStyle.cpp
@@ -71,7 +71,7 @@ const ISlateStyle& FInsightsStyle::Get()
 
 FName FInsightsStyle::GetStyleSetName()
 {
-	static FName StyleSetName(TEXT("AssetManager/InsightsStyle"));
+	static FName StyleSetName(TEXT("AssetManagerEditor/InsightsStyle"));
 	return StyleSetName;
 }
 
