AssetManagerEditor: Removed the duplicated Insights code. Uses instead the TraceInsightsCore module.
Insights / Tree Views: Added separate color for icon and text for tree items.

diff --git a/Engine/Plugins/Editor/AssetManagerEditor/Source/AssetManagerEditor/AssetManagerEditor.Build.cs b/Engine/Plugins/Editor/AssetManagerEditor/Source/AssetManagerEditor/AssetManagerEditor.Build.cs
--- a/Engine/Plugins/Editor/AssetManagerEditor/Source/AssetManagerEditor/AssetManagerEditor.Build.cs
+++ b/Engine/Plugins/Editor/AssetManagerEditor/Source/AssetManagerEditor/AssetManagerEditor.Build.cs
@@ -49,6 +49,7 @@ public class AssetManagerEditor : ModuleRules
 				"TreeMap",
 				"ToolMenus",
 				"ToolWidgets",
+				"TraceInsightsCore",
 				"SourceControl",
 				"SourceControlWindows",
 				"UncontrolledChangelists",
diff --git a/Engine/Plugins/Editor/AssetManagerEditor/Source/AssetManagerEditor/Private/AssetManagerEditorModule.cpp b/Engine/Plugins/Editor/AssetManagerEditor/Source/AssetManagerEditor/Private/AssetManagerEditorModule.cpp
--- a/Engine/Plugins/Editor/AssetManagerEditor/Source/AssetManagerEditor/Private/AssetManagerEditorModule.cpp
+++ b/Engine/Plugins/Editor/AssetManagerEditor/Source/AssetManagerEditor/Private/AssetManagerEditorModule.cpp
@@ -66,8 +66,7 @@
 #include "ContentBrowserMenuContexts.h"
 #include "IContentBrowserDataModule.h"
 #include "ContentBrowserDataSubsystem.h"
-#include "Insights/Common/InsightsStyle.h"
-#include "Insights/Filter/ViewModels/Filters.h"
+#include "InsightsCore/Filter/ViewModels/Filters.h"
 #include "TreeView/AssetTable.h"
 #include "TreeView/SAssetTableTreeView.h"
 
@@ -486,7 +485,6 @@ void FAssetManagerEditorModule::StartupModule()
 
 	if (GIsEditor && !IsRunningCommandlet())
 	{
-		UE::Insights::FInsightsStyle::Initialize();
 		UE::Insights::FFilterService::Initialize();
 
 		AuditCmds.Add(IConsoleManager::Get().RegisterConsoleCommand(
@@ -697,7 +695,6 @@ void FAssetManagerEditorModule::ShutdownModule()
 		UToolMenus::UnregisterOwner(this);
 
 		UE::Insights::FFilterService::Shutdown();
-		UE::Insights::FInsightsStyle::Shutdown();
 	}
 }
 
diff --git a/Engine/Plugins/Editor/AssetManagerEditor/Source/AssetManagerEditor/Private/Insights/Common/AsyncOperationProgress.h b/Engine/Plugins/Editor/AssetManagerEditor/Source/AssetManagerEditor/Private/Insights/Common/AsyncOperationProgress.h
deleted file mode 100644
diff --git a/Engine/Plugins/Editor/AssetManagerEditor/Source/AssetManagerEditor/Private/Insights/Common/AsyncWorkUtils.h b/Engine/Plugins/Editor/AssetManagerEditor/Source/AssetManagerEditor/Private/Insights/Common/AsyncWorkUtils.h
deleted file mode 100644
diff --git a/Engine/Plugins/Editor/AssetManagerEditor/Source/AssetManagerEditor/Private/Insights/Common/InsightsStyle.cpp b/Engine/Plugins/Editor/AssetManagerEditor/Source/AssetManagerEditor/Private/Insights/Common/InsightsStyle.cpp
deleted file mode 100644
diff --git a/Engine/Plugins/Editor/AssetManagerEditor/Source/AssetManagerEditor/Private/Insights/Common/InsightsStyle.h b/Engine/Plugins/Editor/AssetManagerEditor/Source/AssetManagerEditor/Private/Insights/Common/InsightsStyle.h
deleted file mode 100644
diff --git a/Engine/Plugins/Editor/AssetManagerEditor/Source/AssetManagerEditor/Private/Insights/Common/Log.cpp b/Engine/Plugins/Editor/AssetManagerEditor/Source/AssetManagerEditor/Private/Insights/Common/Log.cpp
deleted file mode 100644
diff --git a/Engine/Plugins/Editor/AssetManagerEditor/Source/AssetManagerEditor/Private/Insights/Common/Log.h b/Engine/Plugins/Editor/AssetManagerEditor/Source/AssetManagerEditor/Private/Insights/Common/Log.h
deleted file mode 100644
diff --git a/Engine/Plugins/Editor/AssetManagerEditor/Source/AssetManagerEditor/Private/Insights/Common/SimpleRtti.h b/Engine/Plugins/Editor/AssetManagerEditor/Source/AssetManagerEditor/Private/Insights/Common/SimpleRtti.h
deleted file mode 100644
diff --git a/Engine/Plugins/Editor/AssetManagerEditor/Source/AssetManagerEditor/Private/Insights/Common/Stopwatch.h b/Engine/Plugins/Editor/AssetManagerEditor/Source/AssetManagerEditor/Private/Insights/Common/Stopwatch.h
deleted file mode 100644
diff --git a/Engine/Plugins/Editor/AssetManagerEditor/Source/AssetManagerEditor/Private/Insights/Common/TimeUtils.cpp b/Engine/Plugins/Editor/AssetManagerEditor/Source/AssetManagerEditor/Private/Insights/Common/TimeUtils.cpp
deleted file mode 100644
diff --git a/Engine/Plugins/Editor/AssetManagerEditor/Source/AssetManagerEditor/Private/Insights/Common/TimeUtils.h b/Engine/Plugins/Editor/AssetManagerEditor/Source/AssetManagerEditor/Private/Insights/Common/TimeUtils.h
deleted file mode 100644
diff --git a/Engine/Plugins/Editor/AssetManagerEditor/Source/AssetManagerEditor/Private/Insights/Filter/ViewModels/FilterConfigurator.cpp b/Engine/Plugins/Editor/AssetManagerEditor/Source/AssetManagerEditor/Private/Insights/Filter/ViewModels/FilterConfigurator.cpp
deleted file mode 100644
diff --git a/Engine/Plugins/Editor/AssetManagerEditor/Source/AssetManagerEditor/Private/Insights/Filter/ViewModels/FilterConfigurator.h b/Engine/Plugins/Editor/AssetManagerEditor/Source/AssetManagerEditor/Private/Insights/Filter/ViewModels/FilterConfigurator.h
deleted file mode 100644
diff --git a/Engine/Plugins/Editor/AssetManagerEditor/Source/AssetManagerEditor/Private/Insights/Filter/ViewModels/FilterConfiguratorNode.cpp b/Engine/Plugins/Editor/AssetManagerEditor/Source/AssetManagerEditor/Private/Insights/Filter/ViewModels/FilterConfiguratorNode.cpp
deleted file mode 100644
diff --git a/Engine/Plugins/Editor/AssetManagerEditor/Source/AssetManagerEditor/Private/Insights/Filter/ViewModels/FilterConfiguratorNode.h b/Engine/Plugins/Editor/AssetManagerEditor/Source/AssetManagerEditor/Private/Insights/Filter/ViewModels/FilterConfiguratorNode.h
deleted file mode 100644
diff --git a/Engine/Plugins/Editor/AssetManagerEditor/Source/AssetManagerEditor/Private/Insights/Filter/ViewModels/Filters.cpp b/Engine/Plugins/Editor/AssetManagerEditor/Source/AssetManagerEditor/Private/Insights/Filter/ViewModels/Filters.cpp
deleted file mode 100644
diff --git a/Engine/Plugins/Editor/AssetManagerEditor/Source/AssetManagerEditor/Private/Insights/Filter/ViewModels/Filters.h b/Engine/Plugins/Editor/AssetManagerEditor/Source/AssetManagerEditor/Private/Insights/Filter/ViewModels/Filters.h
deleted file mode 100644
diff --git a/Engine/Plugins/Editor/AssetManagerEditor/Source/AssetManagerEditor/Private/Insights/Filter/ViewModels/IFilterExecutor.h b/Engine/Plugins/Editor/AssetManagerEditor/Source/AssetManagerEditor/Private/Insights/Filter/ViewModels/IFilterExecutor.h
deleted file mode 100644
diff --git a/Engine/Plugins/Editor/AssetManagerEditor/Source/AssetManagerEditor/Private/Insights/Filter/Widgets/SAdvancedFilter.cpp b/Engine/Plugins/Editor/AssetManagerEditor/Source/AssetManagerEditor/Private/Insights/Filter/Widgets/SAdvancedFilter.cpp
deleted file mode 100644
diff --git a/Engine/Plugins/Editor/AssetManagerEditor/Source/AssetManagerEditor/Private/Insights/Filter/Widgets/SAdvancedFilter.h b/Engine/Plugins/Editor/AssetManagerEditor/Source/AssetManagerEditor/Private/Insights/Filter/Widgets/SAdvancedFilter.h
deleted file mode 100644
diff --git a/Engine/Plugins/Editor/AssetManagerEditor/Source/AssetManagerEditor/Private/Insights/Filter/Widgets/SFilterConfigurator.cpp b/Engine/Plugins/Editor/AssetManagerEditor/Source/AssetManagerEditor/Private/Insights/Filter/Widgets/SFilterConfigurator.cpp
deleted file mode 100644
diff --git a/Engine/Plugins/Editor/AssetManagerEditor/Source/AssetManagerEditor/Private/Insights/Filter/Widgets/SFilterConfigurator.h b/Engine/Plugins/Editor/AssetManagerEditor/Source/AssetManagerEditor/Private/Insights/Filter/Widgets/SFilterConfigurator.h
deleted file mode 100644
diff --git a/Engine/Plugins/Editor/AssetManagerEditor/Source/AssetManagerEditor/Private/Insights/Filter/Widgets/SFilterConfiguratorRow.cpp b/Engine/Plugins/Editor/AssetManagerEditor/Source/AssetManagerEditor/Private/Insights/Filter/Widgets/SFilterConfiguratorRow.cpp
deleted file mode 100644
diff --git a/Engine/Plugins/Editor/AssetManagerEditor/Source/AssetManagerEditor/Private/Insights/Filter/Widgets/SFilterConfiguratorRow.h b/Engine/Plugins/Editor/AssetManagerEditor/Source/AssetManagerEditor/Private/Insights/Filter/Widgets/SFilterConfiguratorRow.h
deleted file mode 100644
diff --git a/Engine/Plugins/Editor/AssetManagerEditor/Source/AssetManagerEditor/Private/Insights/Table/ViewModels/BaseTreeNode.cpp b/Engine/Plugins/Editor/AssetManagerEditor/Source/AssetManagerEditor/Private/Insights/Table/ViewModels/BaseTreeNode.cpp
deleted file mode 100644
diff --git a/Engine/Plugins/Editor/AssetManagerEditor/Source/AssetManagerEditor/Private/Insights/Table/ViewModels/BaseTreeNode.h b/Engine/Plugins/Editor/AssetManagerEditor/Source/AssetManagerEditor/Private/Insights/Table/ViewModels/BaseTreeNode.h
deleted file mode 100644
diff --git a/Engine/Plugins/Editor/AssetManagerEditor/Source/AssetManagerEditor/Private/Insights/Table/ViewModels/Table.cpp b/Engine/Plugins/Editor/AssetManagerEditor/Source/AssetManagerEditor/Private/Insights/Table/ViewModels/Table.cpp
deleted file mode 100644
diff --git a/Engine/Plugins/Editor/AssetManagerEditor/Source/AssetManagerEditor/Private/Insights/Table/ViewModels/Table.h b/Engine/Plugins/Editor/AssetManagerEditor/Source/AssetManagerEditor/Private/Insights/Table/ViewModels/Table.h
deleted file mode 100644
diff --git a/Engine/Plugins/Editor/AssetManagerEditor/Source/AssetManagerEditor/Private/Insights/Table/ViewModels/TableCellValue.h b/Engine/Plugins/Editor/AssetManagerEditor/Source/AssetManagerEditor/Private/Insights/Table/ViewModels/TableCellValue.h
deleted file mode 100644
diff --git a/Engine/Plugins/Editor/AssetManagerEditor/Source/AssetManagerEditor/Private/Insights/Table/ViewModels/TableCellValueFormatter.cpp b/Engine/Plugins/Editor/AssetManagerEditor/Source/AssetManagerEditor/Private/Insights/Table/ViewModels/TableCellValueFormatter.cpp
deleted file mode 100644
diff --git a/Engine/Plugins/Editor/AssetManagerEditor/Source/AssetManagerEditor/Private/Insights/Table/ViewModels/TableCellValueFormatter.h b/Engine/Plugins/Editor/AssetManagerEditor/Source/AssetManagerEditor/Private/Insights/Table/ViewModels/TableCellValueFormatter.h
deleted file mode 100644
diff --git a/Engine/Plugins/Editor/AssetManagerEditor/Source/AssetManagerEditor/Private/Insights/Table/ViewModels/TableCellValueGetter.cpp b/Engine/Plugins/Editor/AssetManagerEditor/Source/AssetManagerEditor/Private/Insights/Table/ViewModels/TableCellValueGetter.cpp
deleted file mode 100644
diff --git a/Engine/Plugins/Editor/AssetManagerEditor/Source/AssetManagerEditor/Private/Insights/Table/ViewModels/TableCellValueGetter.h b/Engine/Plugins/Editor/AssetManagerEditor/Source/AssetManagerEditor/Private/Insights/Table/ViewModels/TableCellValueGetter.h
deleted file mode 100644
diff --git a/Engine/Plugins/Editor/AssetManagerEditor/Source/AssetManagerEditor/Private/Insights/Table/ViewModels/TableCellValueSorter.cpp b/Engine/Plugins/Editor/AssetManagerEditor/Source/AssetManagerEditor/Private/Insights/Table/ViewModels/TableCellValueSorter.cpp
deleted file mode 100644
diff --git a/Engine/Plugins/Editor/AssetManagerEditor/Source/AssetManagerEditor/Private/Insights/Table/ViewModels/TableCellValueSorter.h b/Engine/Plugins/Editor/AssetManagerEditor/Source/AssetManagerEditor/Private/Insights/Table/ViewModels/TableCellValueSorter.h
deleted file mode 100644
diff --git a/Engine/Plugins/Editor/AssetManagerEditor/Source/AssetManagerEditor/Private/Insights/Table/ViewModels/TableColumn.cpp b/Engine/Plugins/Editor/AssetManagerEditor/Source/AssetManagerEditor/Private/Insights/Table/ViewModels/TableColumn.cpp
deleted file mode 100644
diff --git a/Engine/Plugins/Editor/AssetManagerEditor/Source/AssetManagerEditor/Private/Insights/Table/ViewModels/TableColumn.h b/Engine/Plugins/Editor/AssetManagerEditor/Source/AssetManagerEditor/Private/Insights/Table/ViewModels/TableColumn.h
deleted file mode 100644
diff --git a/Engine/Plugins/Editor/AssetManagerEditor/Source/AssetManagerEditor/Private/Insights/Table/ViewModels/TableTreeNode.cpp b/Engine/Plugins/Editor/AssetManagerEditor/Source/AssetManagerEditor/Private/Insights/Table/ViewModels/TableTreeNode.cpp
deleted file mode 100644
diff --git a/Engine/Plugins/Editor/AssetManagerEditor/Source/AssetManagerEditor/Private/Insights/Table/ViewModels/TableTreeNode.h b/Engine/Plugins/Editor/AssetManagerEditor/Source/AssetManagerEditor/Private/Insights/Table/ViewModels/TableTreeNode.h
deleted file mode 100644
diff --git a/Engine/Plugins/Editor/AssetManagerEditor/Source/AssetManagerEditor/Private/Insights/Table/ViewModels/TreeNodeGrouping.cpp b/Engine/Plugins/Editor/AssetManagerEditor/Source/AssetManagerEditor/Private/Insights/Table/ViewModels/TreeNodeGrouping.cpp
deleted file mode 100644
diff --git a/Engine/Plugins/Editor/AssetManagerEditor/Source/AssetManagerEditor/Private/Insights/Table/ViewModels/TreeNodeGrouping.h b/Engine/Plugins/Editor/AssetManagerEditor/Source/AssetManagerEditor/Private/Insights/Table/ViewModels/TreeNodeGrouping.h
deleted file mode 100644
diff --git a/Engine/Plugins/Editor/AssetManagerEditor/Source/AssetManagerEditor/Private/Insights/Table/Widgets/SAsyncOperationStatus.cpp b/Engine/Plugins/Editor/AssetManagerEditor/Source/AssetManagerEditor/Private/Insights/Table/Widgets/SAsyncOperationStatus.cpp
deleted file mode 100644
diff --git a/Engine/Plugins/Editor/AssetManagerEditor/Source/AssetManagerEditor/Private/Insights/Table/Widgets/SAsyncOperationStatus.h b/Engine/Plugins/Editor/AssetManagerEditor/Source/AssetManagerEditor/Private/Insights/Table/Widgets/SAsyncOperationStatus.h
deleted file mode 100644
diff --git a/Engine/Plugins/Editor/AssetManagerEditor/Source/AssetManagerEditor/Private/Insights/Table/Widgets/STableTreeView.cpp b/Engine/Plugins/Editor/AssetManagerEditor/Source/AssetManagerEditor/Private/Insights/Table/Widgets/STableTreeView.cpp
deleted file mode 100644
diff --git a/Engine/Plugins/Editor/AssetManagerEditor/Source/AssetManagerEditor/Private/Insights/Table/Widgets/STableTreeView.h b/Engine/Plugins/Editor/AssetManagerEditor/Source/AssetManagerEditor/Private/Insights/Table/Widgets/STableTreeView.h
deleted file mode 100644
diff --git a/Engine/Plugins/Editor/AssetManagerEditor/Source/AssetManagerEditor/Private/Insights/Table/Widgets/STableTreeViewCell.cpp b/Engine/Plugins/Editor/AssetManagerEditor/Source/AssetManagerEditor/Private/Insights/Table/Widgets/STableTreeViewCell.cpp
deleted file mode 100644
diff --git a/Engine/Plugins/Editor/AssetManagerEditor/Source/AssetManagerEditor/Private/Insights/Table/Widgets/STableTreeViewCell.h b/Engine/Plugins/Editor/AssetManagerEditor/Source/AssetManagerEditor/Private/Insights/Table/Widgets/STableTreeViewCell.h
deleted file mode 100644
diff --git a/Engine/Plugins/Editor/AssetManagerEditor/Source/AssetManagerEditor/Private/Insights/Table/Widgets/STableTreeViewRow.cpp b/Engine/Plugins/Editor/AssetManagerEditor/Source/AssetManagerEditor/Private/Insights/Table/Widgets/STableTreeViewRow.cpp
deleted file mode 100644
diff --git a/Engine/Plugins/Editor/AssetManagerEditor/Source/AssetManagerEditor/Private/Insights/Table/Widgets/STableTreeViewRow.h b/Engine/Plugins/Editor/AssetManagerEditor/Source/AssetManagerEditor/Private/Insights/Table/Widgets/STableTreeViewRow.h
deleted file mode 100644
diff --git a/Engine/Plugins/Editor/AssetManagerEditor/Source/AssetManagerEditor/Private/Insights/Table/Widgets/STableTreeViewTooltip.cpp b/Engine/Plugins/Editor/AssetManagerEditor/Source/AssetManagerEditor/Private/Insights/Table/Widgets/STableTreeViewTooltip.cpp
deleted file mode 100644
diff --git a/Engine/Plugins/Editor/AssetManagerEditor/Source/AssetManagerEditor/Private/Insights/Table/Widgets/STableTreeViewTooltip.h b/Engine/Plugins/Editor/AssetManagerEditor/Source/AssetManagerEditor/Private/Insights/Table/Widgets/STableTreeViewTooltip.h
deleted file mode 100644
diff --git a/Engine/Plugins/Editor/AssetManagerEditor/Source/AssetManagerEditor/Private/SAssetAuditBrowser.cpp b/Engine/Plugins/Editor/AssetManagerEditor/Source/AssetManagerEditor/Private/SAssetAuditBrowser.cpp
--- a/Engine/Plugins/Editor/AssetManagerEditor/Source/AssetManagerEditor/Private/SAssetAuditBrowser.cpp
+++ b/Engine/Plugins/Editor/AssetManagerEditor/Source/AssetManagerEditor/Private/SAssetAuditBrowser.cpp
@@ -25,7 +25,7 @@
 
 #include "AssetManagerEditorModule.h"
 #include "AssetManagerEditorCommands.h"
-#include "Insights/Common/Log.h"
+#include "InsightsCore/Common/Log.h"
 #include "TreeView/AssetTable.h"
 #include "TreeView/AssetTreeNode.h"
 #include "TreeView/SAssetTableTreeView.h"
diff --git a/Engine/Plugins/Editor/AssetManagerEditor/Source/AssetManagerEditor/Private/TreeView/AssetDependencyGrouping.cpp b/Engine/Plugins/Editor/AssetManagerEditor/Source/AssetManagerEditor/Private/TreeView/AssetDependencyGrouping.cpp
--- a/Engine/Plugins/Editor/AssetManagerEditor/Source/AssetManagerEditor/Private/TreeView/AssetDependencyGrouping.cpp
+++ b/Engine/Plugins/Editor/AssetManagerEditor/Source/AssetManagerEditor/Private/TreeView/AssetDependencyGrouping.cpp
@@ -3,8 +3,8 @@
 #include "AssetDependencyGrouping.h"
 #include "AssetTreeNode.h"
 #include "AssetTable.h"
-#include "Insights/Common/AsyncOperationProgress.h"
-#include "Insights/Common/Log.h"
+#include "InsightsCore/Common/AsyncOperationProgress.h"
+#include "InsightsCore/Common/Log.h"
 
 #define LOCTEXT_NAMESPACE "FAssetDependencyGrouping"
 
diff --git a/Engine/Plugins/Editor/AssetManagerEditor/Source/AssetManagerEditor/Private/TreeView/AssetDependencyGrouping.h b/Engine/Plugins/Editor/AssetManagerEditor/Source/AssetManagerEditor/Private/TreeView/AssetDependencyGrouping.h
--- a/Engine/Plugins/Editor/AssetManagerEditor/Source/AssetManagerEditor/Private/TreeView/AssetDependencyGrouping.h
+++ b/Engine/Plugins/Editor/AssetManagerEditor/Source/AssetManagerEditor/Private/TreeView/AssetDependencyGrouping.h
@@ -3,8 +3,8 @@
 #pragma once
 
 #include "CoreMinimal.h"
-#include "Insights/Table/ViewModels/BaseTreeNode.h"
-#include "Insights/Table/ViewModels/TreeNodeGrouping.h"
+#include "InsightsCore/Table/ViewModels/BaseTreeNode.h"
+#include "InsightsCore/Table/ViewModels/TreeNodeGrouping.h"
 
 class FAssetDependencyGrouping : public UE::Insights::FTreeNodeGrouping
 {
diff --git a/Engine/Plugins/Editor/AssetManagerEditor/Source/AssetManagerEditor/Private/TreeView/AssetTable.cpp b/Engine/Plugins/Editor/AssetManagerEditor/Source/AssetManagerEditor/Private/TreeView/AssetTable.cpp
--- a/Engine/Plugins/Editor/AssetManagerEditor/Source/AssetManagerEditor/Private/TreeView/AssetTable.cpp
+++ b/Engine/Plugins/Editor/AssetManagerEditor/Source/AssetManagerEditor/Private/TreeView/AssetTable.cpp
@@ -9,12 +9,12 @@
 
 // AssetManagerEditor
 #include "AssetTreeNode.h"
-#include "Insights/Common/Log.h"
-#include "Insights/Common/Stopwatch.h"
-#include "Insights/Table/ViewModels/TableCellValueFormatter.h"
-#include "Insights/Table/ViewModels/TableCellValueGetter.h"
-#include "Insights/Table/ViewModels/TableCellValueSorter.h"
-#include "Insights/Table/ViewModels/TableColumn.h"
+#include "InsightsCore/Common/Log.h"
+#include "InsightsCore/Common/Stopwatch.h"
+#include "InsightsCore/Table/ViewModels/TableCellValueFormatter.h"
+#include "InsightsCore/Table/ViewModels/TableCellValueGetter.h"
+#include "InsightsCore/Table/ViewModels/TableCellValueSorter.h"
+#include "InsightsCore/Table/ViewModels/TableColumn.h"
 
 #define LOCTEXT_NAMESPACE "FAssetTable"
 
diff --git a/Engine/Plugins/Editor/AssetManagerEditor/Source/AssetManagerEditor/Private/TreeView/AssetTable.h b/Engine/Plugins/Editor/AssetManagerEditor/Source/AssetManagerEditor/Private/TreeView/AssetTable.h
--- a/Engine/Plugins/Editor/AssetManagerEditor/Source/AssetManagerEditor/Private/TreeView/AssetTable.h
+++ b/Engine/Plugins/Editor/AssetManagerEditor/Source/AssetManagerEditor/Private/TreeView/AssetTable.h
@@ -6,7 +6,7 @@
 
 #include "AssetRegistry/AssetData.h"
 #include "Containers/StringView.h"
-#include "Insights/Table/ViewModels/Table.h"
+#include "InsightsCore/Table/ViewModels/Table.h"
 #include "Misc/CString.h"
 #include "Misc/TVariant.h"
 
diff --git a/Engine/Plugins/Editor/AssetManagerEditor/Source/AssetManagerEditor/Private/TreeView/AssetTreeNode.cpp b/Engine/Plugins/Editor/AssetManagerEditor/Source/AssetManagerEditor/Private/TreeView/AssetTreeNode.cpp
--- a/Engine/Plugins/Editor/AssetManagerEditor/Source/AssetManagerEditor/Private/TreeView/AssetTreeNode.cpp
+++ b/Engine/Plugins/Editor/AssetManagerEditor/Source/AssetManagerEditor/Private/TreeView/AssetTreeNode.cpp
@@ -3,9 +3,9 @@
 #include "AssetTreeNode.h"
 
 #include "AssetDependencyGrouping.h"
-#include "Insights/Common/AsyncOperationProgress.h"
-#include "Insights/Common/InsightsStyle.h"
-#include "Insights/Table/Widgets/STableTreeView.h"
+#include "InsightsCore/Common/AsyncOperationProgress.h"
+#include "InsightsCore/Common/InsightsCoreStyle.h"
+#include "InsightsCore/Table/Widgets/STableTreeView.h"
 #include "Internationalization/Internationalization.h"
 #include "Styling/StyleColors.h"
 
@@ -28,27 +28,27 @@ const FSlateBrush* FAssetTreeNode::GetIcon(EStyle Style) const
 		default:
 		case EStyle::EDefault:
 		{
-			return UE::Insights::FInsightsStyle::GetBrush("Icons.Leaf.TreeItem");
+			return UE::Insights::FInsightsCoreStyle::GetBrush("Icons.Leaf.TreeItem");
 		}
 
 		case EStyle::EAsset:
 		{
-			return UE::Insights::FInsightsStyle::GetBrush("Icons.Asset.TreeItem");
+			return UE::Insights::FInsightsCoreStyle::GetBrush("Icons.Asset.TreeItem");
 		}
 				
 		case EStyle::EGroup:
 		{
-			return UE::Insights::FInsightsStyle::GetBrush("Icons.Group.TreeItem");
+			return UE::Insights::FInsightsCoreStyle::GetBrush("Icons.Group.TreeItem");
 		}	
 
 		case EStyle::EDependencies:
 		{
-			return UE::Insights::FInsightsStyle::GetBrush("Icons.Dependencies.TreeItem");
+			return UE::Insights::FInsightsCoreStyle::GetBrush("Icons.Dependencies.TreeItem");
 		}
 
 		case EStyle::EPlugin:
 		{
-			return UE::Insights::FInsightsStyle::GetBrush("Icons.Plugin.TreeItem");
+			return UE::Insights::FInsightsCoreStyle::GetBrush("Icons.Plugin.TreeItem");
 		}
 	}
 }
diff --git a/Engine/Plugins/Editor/AssetManagerEditor/Source/AssetManagerEditor/Private/TreeView/AssetTreeNode.h b/Engine/Plugins/Editor/AssetManagerEditor/Source/AssetManagerEditor/Private/TreeView/AssetTreeNode.h
--- a/Engine/Plugins/Editor/AssetManagerEditor/Source/AssetManagerEditor/Private/TreeView/AssetTreeNode.h
+++ b/Engine/Plugins/Editor/AssetManagerEditor/Source/AssetManagerEditor/Private/TreeView/AssetTreeNode.h
@@ -5,7 +5,7 @@
 #include "CoreMinimal.h"
 
 #include "AssetTable.h"
-#include "Insights/Table/ViewModels/TableTreeNode.h"
+#include "InsightsCore/Table/ViewModels/TableTreeNode.h"
 
 ////////////////////////////////////////////////////////////////////////////////////////////////////
 
diff --git a/Engine/Plugins/Editor/AssetManagerEditor/Source/AssetManagerEditor/Private/TreeView/SAssetTableTreeView.cpp b/Engine/Plugins/Editor/AssetManagerEditor/Source/AssetManagerEditor/Private/TreeView/SAssetTableTreeView.cpp
--- a/Engine/Plugins/Editor/AssetManagerEditor/Source/AssetManagerEditor/Private/TreeView/SAssetTableTreeView.cpp
+++ b/Engine/Plugins/Editor/AssetManagerEditor/Source/AssetManagerEditor/Private/TreeView/SAssetTableTreeView.cpp
@@ -16,13 +16,12 @@
 #include "HAL/FileManager.h"
 #include "HAL/PlatformFileManager.h"
 #include "IContentBrowserSingleton.h"
-#include "Insights/Common/InsightsStyle.h"
-#include "Insights/Common/Log.h"
-#include "Insights/Common/Stopwatch.h"
-#include "Insights/Filter/ViewModels/FilterConfigurator.h"
-#include "Insights/Table/ViewModels/TableCellValueFormatter.h"
-#include "Insights/Table/ViewModels/TableColumn.h"
-#include "Insights/Table/ViewModels/TreeNodeGrouping.h"
+#include "InsightsCore/Common/Log.h"
+#include "InsightsCore/Common/Stopwatch.h"
+#include "InsightsCore/Filter/ViewModels/FilterConfigurator.h"
+#include "InsightsCore/Table/ViewModels/TableCellValueFormatter.h"
+#include "InsightsCore/Table/ViewModels/TableColumn.h"
+#include "InsightsCore/Table/ViewModels/TreeNodeGrouping.h"
 #include "Interfaces/IPluginManager.h"
 #include "Logging/MessageLog.h"
 #include "Misc/FileHelper.h"
@@ -51,6 +50,7 @@ extern UNREALED_API UEditorEngine* GEditor;
 
 SAssetTableTreeView::SAssetTableTreeView()
 {
+	MaxDepthToAutoExpand = 1;
 }
 
 ////////////////////////////////////////////////////////////////////////////////////////////////////
diff --git a/Engine/Plugins/Editor/AssetManagerEditor/Source/AssetManagerEditor/Private/TreeView/SAssetTableTreeView.h b/Engine/Plugins/Editor/AssetManagerEditor/Source/AssetManagerEditor/Private/TreeView/SAssetTableTreeView.h
--- a/Engine/Plugins/Editor/AssetManagerEditor/Source/AssetManagerEditor/Private/TreeView/SAssetTableTreeView.h
+++ b/Engine/Plugins/Editor/AssetManagerEditor/Source/AssetManagerEditor/Private/TreeView/SAssetTableTreeView.h
@@ -9,7 +9,7 @@
 #include "AssetManagerEditorModule.h"
 
 #include "TreeView/AssetTable.h"
-#include "Insights/Table/Widgets/STableTreeView.h"
+#include "InsightsCore/Table/Widgets/STableTreeView.h"
 
 class FAssetTreeNode;
 
diff --git a/Engine/Source/Developer/TraceInsightsCore/Private/InsightsCore/Common/InsightsCoreStyle.cpp b/Engine/Source/Developer/TraceInsightsCore/Private/InsightsCore/Common/InsightsCoreStyle.cpp
--- a/Engine/Source/Developer/TraceInsightsCore/Private/InsightsCore/Common/InsightsCoreStyle.cpp
+++ b/Engine/Source/Developer/TraceInsightsCore/Private/InsightsCore/Common/InsightsCoreStyle.cpp
@@ -22,6 +22,13 @@ TSharedPtr<FInsightsCoreStyle::FStyle> FInsightsCoreStyle::StyleInstance = nullp
 
 ////////////////////////////////////////////////////////////////////////////////////////////////////
 
+const ISlateStyle& FInsightsCoreStyle::Get()
+{
+	return *StyleInstance;
+}
+
+////////////////////////////////////////////////////////////////////////////////////////////////////
+
 void FInsightsCoreStyle::Initialize()
 {
 	LLM_SCOPE_BYNAME(TEXT("Insights/Style"));
@@ -62,7 +69,7 @@ void FInsightsCoreStyle::Shutdown()
 
 FName FInsightsCoreStyle::GetStyleSetName()
 {
-	static FName StyleSetName(TEXT("AssetManager/InsightsStyle"));
+	static FName StyleSetName(TEXT("TraceInsightsCore/InsightsStyle"));
 	return StyleSetName;
 }
 
diff --git a/Engine/Source/Developer/TraceInsightsCore/Private/InsightsCore/Table/ViewModels/BaseTreeNode.cpp b/Engine/Source/Developer/TraceInsightsCore/Private/InsightsCore/Table/ViewModels/BaseTreeNode.cpp
--- a/Engine/Source/Developer/TraceInsightsCore/Private/InsightsCore/Table/ViewModels/BaseTreeNode.cpp
+++ b/Engine/Source/Developer/TraceInsightsCore/Private/InsightsCore/Table/ViewModels/BaseTreeNode.cpp
@@ -73,11 +73,11 @@ FLinearColor FBaseTreeNode::GetDefaultColor(bool bIsGroupNode)
 {
 	if (bIsGroupNode)
 	{
-		return FLinearColor(1.0f, 0.9f, 0.6f, 1.0f);
+		return USlateThemeManager::Get().GetColor(EStyleColor::AccentYellow);
 	}
 	else
 	{
-		return FLinearColor(1.0f, 1.0f, 1.0f, 1.0f);
+		return USlateThemeManager::Get().GetColor(EStyleColor::AccentWhite);
 	}
 }
 
diff --git a/Engine/Source/Developer/TraceInsightsCore/Private/InsightsCore/Table/Widgets/STableTreeViewCell.cpp b/Engine/Source/Developer/TraceInsightsCore/Private/InsightsCore/Table/Widgets/STableTreeViewCell.cpp
--- a/Engine/Source/Developer/TraceInsightsCore/Private/InsightsCore/Table/Widgets/STableTreeViewCell.cpp
+++ b/Engine/Source/Developer/TraceInsightsCore/Private/InsightsCore/Table/Widgets/STableTreeViewCell.cpp
@@ -180,8 +180,8 @@ FSlateColor STableTreeViewCell::GetIconColorAndOpacity() const
 	bool bIsHoveredOrSelected = IsHovered() || IsSelected();
 
 	return bIsHoveredOrSelected ?
-		TableTreeNodePtr->GetColor() :
-		TableTreeNodePtr->GetColor().CopyWithNewOpacity(0.5f);
+		TableTreeNodePtr->GetIconColor() :
+		TableTreeNodePtr->GetIconColor().CopyWithNewOpacity(0.8f);
 }
 
 ////////////////////////////////////////////////////////////////////////////////////////////////////
diff --git a/Engine/Source/Developer/TraceInsightsCore/Public/InsightsCore/Common/InsightsCoreStyle.h b/Engine/Source/Developer/TraceInsightsCore/Public/InsightsCore/Common/InsightsCoreStyle.h
--- a/Engine/Source/Developer/TraceInsightsCore/Public/InsightsCore/Common/InsightsCoreStyle.h
+++ b/Engine/Source/Developer/TraceInsightsCore/Public/InsightsCore/Common/InsightsCoreStyle.h
@@ -23,11 +23,7 @@ public:
 	static void Shutdown();
 
 	TRACEINSIGHTSCORE_API static FName GetStyleSetName();
-
-	static const ISlateStyle& Get()
-	{
-		return *StyleInstance;
-	}
+	TRACEINSIGHTSCORE_API static const ISlateStyle& Get();
 
 	static const FLinearColor& GetColor(FName PropertyName, const ANSICHAR* Specifier = NULL)
 	{
@@ -77,7 +73,7 @@ public:
 private:
 	static TSharedRef<FInsightsCoreStyle::FStyle> Create();
 
-	static TSharedPtr<FInsightsCoreStyle::FStyle> StyleInstance;
+	TRACEINSIGHTSCORE_API static TSharedPtr<FInsightsCoreStyle::FStyle> StyleInstance;
 };
 
 } // namespace UE::Insights
diff --git a/Engine/Source/Developer/TraceInsightsCore/Public/InsightsCore/Common/Log.h b/Engine/Source/Developer/TraceInsightsCore/Public/InsightsCore/Common/Log.h
--- a/Engine/Source/Developer/TraceInsightsCore/Public/InsightsCore/Common/Log.h
+++ b/Engine/Source/Developer/TraceInsightsCore/Public/InsightsCore/Common/Log.h
@@ -6,6 +6,6 @@
 #include "HAL/LowLevelMemTracker.h"
 #include "Logging/LogMacros.h"
 
-DECLARE_LOG_CATEGORY_EXTERN(LogInsights, Verbose, All);
+TRACEINSIGHTSCORE_API DECLARE_LOG_CATEGORY_EXTERN(LogInsights, Verbose, All);
 
 LLM_DECLARE_TAG_API(Insights, TRACEINSIGHTSCORE_API);
diff --git a/Engine/Source/Developer/TraceInsightsCore/Public/InsightsCore/Table/ViewModels/BaseTreeNode.h b/Engine/Source/Developer/TraceInsightsCore/Public/InsightsCore/Table/ViewModels/BaseTreeNode.h
--- a/Engine/Source/Developer/TraceInsightsCore/Public/InsightsCore/Table/ViewModels/BaseTreeNode.h
+++ b/Engine/Source/Developer/TraceInsightsCore/Public/InsightsCore/Table/ViewModels/BaseTreeNode.h
@@ -123,17 +123,25 @@ public:
 		return GetDefaultIcon(IsGroup());
 	}
 
+	/**
+	 * @returns the color tint to be used for the icon used for this node.
+	 */
+	virtual FLinearColor GetIconColor() const
+	{
+		return GetDefaultColor(IsGroup());
+	}
+
 	/**
 	 * @returns the default color tint for a group/leaf node.
 	 */
 	static FLinearColor GetDefaultColor(bool bIsGroupNode);
 
 	/**
-	 * @returns the color tint to be used for the icon and the name text of this node.
+	 * @returns the color tint to be used for the name text of this node.
 	 */
 	virtual FLinearColor GetColor() const
 	{
-		return GetDefaultColor(IsGroup());
+		return FLinearColor(1.0f, 1.0f, 1.0f, 1.0f);
 	}
 
 	//////////////////////////////////////////////////
