AssetManagerEditor: Re-enabled async operations (filtering, sorting, grouping) in AssetTableTreeView (Asset Disk Size tool).

#rnx

diff --git a/Engine/Plugins/Editor/AssetManagerEditor/Source/AssetManagerEditor/Private/TreeView/SAssetTableTreeView.cpp b/Engine/Plugins/Editor/AssetManagerEditor/Source/AssetManagerEditor/Private/TreeView/SAssetTableTreeView.cpp
--- a/Engine/Plugins/Editor/AssetManagerEditor/Source/AssetManagerEditor/Private/TreeView/SAssetTableTreeView.cpp
+++ b/Engine/Plugins/Editor/AssetManagerEditor/Source/AssetManagerEditor/Private/TreeView/SAssetTableTreeView.cpp
@@ -51,6 +51,7 @@ extern UNREALED_API UEditorEngine* GEditor;
 SAssetTableTreeView::SAssetTableTreeView()
 {
 	MaxDepthToAutoExpand = 1;
+	bRunInAsyncMode = true;
 }
 
 ////////////////////////////////////////////////////////////////////////////////////////////////////
