[Reference Viewer] More contextual icons for resolved properties

diff --git a/Engine/Plugins/Editor/AssetManagerEditor/Source/AssetManagerEditor/Private/ReferenceViewer/EdGraph_ReferenceViewer.cpp b/Engine/Plugins/Editor/AssetManagerEditor/Source/AssetManagerEditor/Private/ReferenceViewer/EdGraph_ReferenceViewer.cpp
--- a/Engine/Plugins/Editor/AssetManagerEditor/Source/AssetManagerEditor/Private/ReferenceViewer/EdGraph_ReferenceViewer.cpp
+++ b/Engine/Plugins/Editor/AssetManagerEditor/Source/AssetManagerEditor/Private/ReferenceViewer/EdGraph_ReferenceViewer.cpp
@@ -356,10 +356,10 @@ TArray<FReferencingPropertyDescription> UEdGraph_ReferenceViewer::RetrieveRefere
 	TArray<FReferencingPropertyDescription> ReferencingProperties;
 
 	// Registering referencing properties to the output array. Property type defaults to EReferencedPropertyType::Property
-	auto AddReferencingProperty = [&ReferencingProperties](const FString& InPropertyName, const FString& InReferencerName,
-		FReferencingPropertyDescription::EAssetReferenceType InPropertyType = FReferencingPropertyDescription::EAssetReferenceType::Property)
+	auto AddReferencingProperty = [&ReferencingProperties, &InReferencedAsset](const FString& InPropertyName, const FString& InReferencerName
+		, FReferencingPropertyDescription::EAssetReferenceType InPropertyType = FReferencingPropertyDescription::EAssetReferenceType::Property)
 	{
-		FReferencingPropertyDescription PropertyDescription(InPropertyName, InReferencerName, InPropertyType);
+		FReferencingPropertyDescription PropertyDescription(InPropertyName, InReferencerName, InPropertyType, InReferencedAsset->GetClass());
 
 		if (!ReferencingProperties.Contains(PropertyDescription))
 		{
diff --git a/Engine/Plugins/Editor/AssetManagerEditor/Source/AssetManagerEditor/Private/ReferenceViewer/SReferencedPropertiesNode.cpp b/Engine/Plugins/Editor/AssetManagerEditor/Source/AssetManagerEditor/Private/ReferenceViewer/SReferencedPropertiesNode.cpp
--- a/Engine/Plugins/Editor/AssetManagerEditor/Source/AssetManagerEditor/Private/ReferenceViewer/SReferencedPropertiesNode.cpp
+++ b/Engine/Plugins/Editor/AssetManagerEditor/Source/AssetManagerEditor/Private/ReferenceViewer/SReferencedPropertiesNode.cpp
@@ -136,18 +136,25 @@ FText SReferencedPropertyNode::GetTooltipText() const
 
 const FSlateBrush* SReferencedPropertyNode::GetIconBrush() const
 {
-	const FSlateBrush* ComponentIcon = FSlateIconFinder::FindIconBrushForClass(UObject::StaticClass());
 
-	switch (PropertyDescription.GetType())
+	const UClass* Class = PropertyDescription.GetPropertyClass();
+
+	const FSlateBrush* ComponentIcon;
+
+	if (Class)
+	{
+		ComponentIcon = FSlateIconFinder::FindIconBrushForClass(Class);
+	}
+	else
 	{
-		case FReferencingPropertyDescription::EAssetReferenceType::Component:
+		if (PropertyDescription.GetType() == FReferencingPropertyDescription::EAssetReferenceType::Component)
+		{
 			ComponentIcon = FSlateIconFinder::FindIconBrushForClass(UActorComponent::StaticClass(), TEXT("SCS.Component"));
-			break;
-		case FReferencingPropertyDescription::EAssetReferenceType::Property:
-		case FReferencingPropertyDescription::EAssetReferenceType::Value:
-		case FReferencingPropertyDescription::EAssetReferenceType::None:
-			break;
-		default: ;
+		}
+		else
+		{
+			ComponentIcon = FSlateIconFinder::FindIconBrushForClass(UObject::StaticClass());
+		}
 	}
 
 	return ComponentIcon;
diff --git a/Engine/Plugins/Editor/AssetManagerEditor/Source/AssetManagerEditor/Public/ReferenceViewer/EdGraphNode_ReferencedProperties.h b/Engine/Plugins/Editor/AssetManagerEditor/Source/AssetManagerEditor/Public/ReferenceViewer/EdGraphNode_ReferencedProperties.h
--- a/Engine/Plugins/Editor/AssetManagerEditor/Source/AssetManagerEditor/Public/ReferenceViewer/EdGraphNode_ReferencedProperties.h
+++ b/Engine/Plugins/Editor/AssetManagerEditor/Source/AssetManagerEditor/Public/ReferenceViewer/EdGraphNode_ReferencedProperties.h
@@ -23,10 +23,11 @@ struct FReferencingPropertyDescription
 		None
 	};
 
-	FReferencingPropertyDescription(const FString& InName, const FString& InReferencerName, const EAssetReferenceType& InType) :
+	FReferencingPropertyDescription(const FString& InName, const FString& InReferencerName, const EAssetReferenceType& InType, const UClass* InClass) :
 		Name(InName),
 		ReferencerName(InReferencerName),
-		Type(InType)
+		Type(InType),
+		PropertyClass(InClass)
 	{
 	}
 
@@ -49,6 +50,8 @@ struct FReferencingPropertyDescription
 	/** Returns the property type as a string (useful e.g. for tooltips) */
 	FString GetTypeAsString() const;
 
+	const UClass* GetPropertyClass() const { return PropertyClass; }
+
 private:
 	friend class SReferencedPropertyNode;
 	FReferencingPropertyDescription() = default;
@@ -56,6 +59,8 @@ private:
 	FString Name;
 	FString ReferencerName;
 	EAssetReferenceType Type = EAssetReferenceType::None;
+
+	const UClass* PropertyClass = nullptr;
 };
 
 
