[Reference Viewer] Get rid of "ghost" referencing properties from properties nodes + properties duplicates

diff --git a/Engine/Plugins/Editor/AssetManagerEditor/Source/AssetManagerEditor/Private/ReferenceViewer/EdGraph_ReferenceViewer.cpp b/Engine/Plugins/Editor/AssetManagerEditor/Source/AssetManagerEditor/Private/ReferenceViewer/EdGraph_ReferenceViewer.cpp
--- a/Engine/Plugins/Editor/AssetManagerEditor/Source/AssetManagerEditor/Private/ReferenceViewer/EdGraph_ReferenceViewer.cpp
+++ b/Engine/Plugins/Editor/AssetManagerEditor/Source/AssetManagerEditor/Private/ReferenceViewer/EdGraph_ReferenceViewer.cpp
@@ -558,7 +558,58 @@ TArray<FReferencingPropertyDescription> UEdGraph_ReferenceViewer::RetrieveRefere
 			const FReferenceChainSearch::FGraphNode* ObjectRefChainLink = ObjectRefChain->GetNode(NodeIndex);
 			UObject* LinkObject = ObjectRefChainLink->ObjectInfo->TryResolveObject();
 
-			if (LinkObject->GetOuter() == InReferencedAsset)
+			if (!LinkObject)
+			{
+				continue;
+			}
+
+			// Let's skip the referenced asset itself
+			if (LinkObject == InReferencedAsset || LinkObject->GetOuter() == InReferencedAsset)
+			{
+				continue;
+			}
+
+			bool bReferencerFound = false;
+
+			UClass* LinkObjClass = LinkObject->GetClass();
+
+			// Check Class or Class Generated by for BPs
+			if (LinkObjClass == InReferencer->GetClass() ||
+				LinkObjClass->ClassGeneratedBy == InReferencer)
+			{
+				bReferencerFound = true;
+			}
+
+			FString IntermediateObjectName;
+
+			if (!bReferencerFound)
+			{
+				// Let's check outer as well, in case the referencing property is a field of an Object/Component
+				if (const UObject* LinkObjectOuter = LinkObject->GetOuter())
+				{
+					// Check Class or Class Generated by for outer
+					if (LinkObjectOuter->GetClass() == InReferencer->GetClass())
+					{
+						bReferencerFound = true;
+					}
+					else if (const UClass* LinkObjectOuterClass = LinkObjectOuter->GetClass())
+					{
+						if (LinkObjectOuterClass->ClassGeneratedBy == InReferencer)
+						{
+							bReferencerFound = true;
+						}
+					}
+
+					if (bReferencerFound)
+					{
+						// Store the Object name, so we can display clearer information
+						// TODO: maybe improve the way this info is collected in FReferencingPropertyDescription
+						IntermediateObjectName = LinkObject->GetName();
+					}
+				}
+			}
+
+			if (!bReferencerFound)
 			{
 				continue;
 			}
@@ -578,13 +629,14 @@ TArray<FReferencingPropertyDescription> UEdGraph_ReferenceViewer::RetrieveRefere
 			{
 				if (ReferenceObject == InReferencedAsset)
 				{
-					const FString RefName = ReferenceInfo->ReferencerName.ToString();
+					FString RefName = ReferenceInfo->ReferencerName.ToString();
 
-					// These should not interest the user
-					if (RefName == "ClassGeneratedBy"
-						|| RefName == "ScriptAndPropertyObjectReferences")
+					if (!IntermediateObjectName.IsEmpty())
 					{
-						continue;
+						// TODO: can be improved by better design.
+						// This is a quick way to display useful info,
+						// e.g. MyStaticMeshComponent:StaticMesh instead of just StaticMesh
+						RefName = IntermediateObjectName + TEXT(":") + RefName;
 					}
 
 					FoundProperties.Add(RefName);
