#UE Do not crash when there are no found assets after typing in a path in the address bar of the reference viewer

diff --git a/Engine/Plugins/Editor/AssetManagerEditor/Source/AssetManagerEditor/Private/AssetManagerEditorModule.cpp b/Engine/Plugins/Editor/AssetManagerEditor/Source/AssetManagerEditor/Private/AssetManagerEditorModule.cpp
--- a/Engine/Plugins/Editor/AssetManagerEditor/Source/AssetManagerEditor/Private/AssetManagerEditorModule.cpp
+++ b/Engine/Plugins/Editor/AssetManagerEditor/Source/AssetManagerEditor/Private/AssetManagerEditorModule.cpp
@@ -1096,67 +1096,70 @@ uint32 FAssetManagerEditorModule::GetHashFromAssetsSelection(const TArray<FAsset
 {
 	TRACE_CPUPROFILER_EVENT_SCOPE(FAssetManagerEditorModule::GetHashFromAssetsSelection);
 
-	TArray<FName> PackageNames;
-	for (const FAssetIdentifier& AssetIdentifier : InAssetIdentifiers)
-	{
-		PackageNames.Add(AssetIdentifier.PackageName);
-	}
-
-	TMap<FName, FAssetData> PackageToAssetDataMap;
-	UE::AssetRegistry::GetAssetForPackages(PackageNames, PackageToAssetDataMap);
-
+	uint32 SelectionHash = 0;
 	FString TabLabel;
 
-	if (!PackageToAssetDataMap.IsEmpty())
+	if (InAssetIdentifiers.Num() > 0)
 	{
-		TabLabel = PackageToAssetDataMap[PackageNames[0]].AssetName.ToString();
-	}
+		TArray<FName> PackageNames;
+		for (const FAssetIdentifier& AssetIdentifier : InAssetIdentifiers)
+		{
+			PackageNames.Add(AssetIdentifier.PackageName);
+		}
 
-	bool bNoAsset = false;
-	// C++ classes will lead to no asset, so we retrieve their package instead
-	// (TODO: this needs to be addressed in order to properly show reference viewer graph for C++ Assets)
+		TMap<FName, FAssetData> PackageToAssetDataMap;
+		UE::AssetRegistry::GetAssetForPackages(PackageNames, PackageToAssetDataMap);
 
-	if (!PackageToAssetDataMap.Contains(PackageNames[0]))
-	{
-		bNoAsset = true;
-		if (TabLabel.IsEmpty())
+		if (!PackageToAssetDataMap.IsEmpty())
 		{
-			TabLabel = PackageNames[0].ToString();
+			TabLabel = PackageToAssetDataMap[PackageNames[0]].AssetName.ToString();
+		}
 
-			if (!InAssetIdentifiers[0].ValueName.IsNone())
+		bool bNoAsset = false;
+		// C++ classes will lead to no asset, so we retrieve their package instead
+		// (TODO: this needs to be addressed in order to properly show reference viewer graph for C++ Assets)
+
+		if (!PackageToAssetDataMap.Contains(PackageNames[0]))
+		{
+			bNoAsset = true;
+			if (TabLabel.IsEmpty())
 			{
-				TabLabel += TEXT(":") + InAssetIdentifiers[0].ValueName.ToString();
+				TabLabel = PackageNames[0].ToString();
+
+				if (!InAssetIdentifiers[0].ValueName.IsNone())
+				{
+					TabLabel += TEXT(":") + InAssetIdentifiers[0].ValueName.ToString();
+				}
 			}
 		}
-	}
 
-	// Label for multiple assets matches Path field at the top of Reference Viewer graph
-	if (InAssetIdentifiers.Num() > 1)
-	{
-		TabLabel += TEXT(" and ") + FString::FromInt(InAssetIdentifiers.Num() - 1) + TEXT(" others");
-	}
+		// Label for multiple assets matches Path field at the top of Reference Viewer graph
+		if (InAssetIdentifiers.Num() > 1)
+		{
+			TabLabel += TEXT(" and ") + FString::FromInt(InAssetIdentifiers.Num() - 1) + TEXT(" others");
+		}
 
-	// Create a hash from the concatenation of package names from all the selected assets. This hash is used to match a selection with a Tab ID
-	// This allows to ignore selection order when comparing selections while looking for an existing Reference Viewer for the current selection
-	// We sort package names, so that selection order is not be taken into account
+		// Create a hash from the concatenation of package names from all the selected assets. This hash is used to match a selection with a Tab ID
+		// This allows to ignore selection order when comparing selections while looking for an existing Reference Viewer for the current selection
+		// We sort package names, so that selection order is not be taken into account
 
-	PackageNames.Sort([](const FName& NameA, const FName& NameB)
-	{ 
-		return NameA.FastLess(NameB);
-	});
+		PackageNames.Sort([](const FName& NameA, const FName& NameB)
+		{ 
+			return NameA.FastLess(NameB);
+		});
 
-	uint32 SelectionHash = 0;
-	for (const FName PackageName : PackageNames)
-	{
-		// In case there is no asset, we might get the same Hash for potentially different graphs (e.g. GameplayTags),
-		// so let's use the tab label, which in that specific case should include both package name and value
-		if (bNoAsset)
+		for (const FName PackageName : PackageNames)
 		{
-			SelectionHash ^= GetTypeHash(TabLabel);
-		}
-		else
-		{
-			SelectionHash ^= GetTypeHash(PackageName);
+			// In case there is no asset, we might get the same Hash for potentially different graphs (e.g. GameplayTags),
+			// so let's use the tab label, which in that specific case should include both package name and value
+			if (bNoAsset)
+			{
+				SelectionHash ^= GetTypeHash(TabLabel);
+			}
+			else
+			{
+				SelectionHash ^= GetTypeHash(PackageName);
+			}
 		}
 	}
 
