[Reference Viewer] You can no longer expand the reference viewer overflow nodes

diff --git a/Engine/Plugins/Editor/AssetManagerEditor/Source/AssetManagerEditor/Private/ReferenceViewer/SReferenceViewer.cpp b/Engine/Plugins/Editor/AssetManagerEditor/Source/AssetManagerEditor/Private/ReferenceViewer/SReferenceViewer.cpp
--- a/Engine/Plugins/Editor/AssetManagerEditor/Source/AssetManagerEditor/Private/ReferenceViewer/SReferenceViewer.cpp
+++ b/Engine/Plugins/Editor/AssetManagerEditor/Source/AssetManagerEditor/Private/ReferenceViewer/SReferenceViewer.cpp
@@ -713,6 +713,7 @@ void SReferenceViewer::OnNodeDoubleClicked(UEdGraphNode* Node)
 
 	const TArray<FAssetIdentifier> CurrentlyVisualizedAssets = GraphObj->GetCurrentGraphRootIdentifiers();
 
+	bool bDependency = false;
 	UEdGraphNode* ParentNode = nullptr;
 	if (UEdGraphNode_Reference* ReferenceNode = Cast<UEdGraphNode_Reference>(Node))
 	{
@@ -725,6 +726,7 @@ void SReferenceViewer::OnNodeDoubleClicked(UEdGraphNode* Node)
 			}
 			else if (ReferenceNode->GetDependencyPin()->LinkedTo.Num() > 0)
 			{
+				bDependency = true;
 				ParentNode = ReferenceNode->GetDependencyPin()->LinkedTo[0]->GetOwningNode();
 			}
 		}
@@ -736,7 +738,7 @@ void SReferenceViewer::OnNodeDoubleClicked(UEdGraphNode* Node)
 		if (UEdGraphNode_Reference* ParentReferenceNode = Cast<UEdGraphNode_Reference>(ParentNode))
 		{
 			FAssetIdentifier ParentID = ParentReferenceNode->GetIdentifier();
-			GraphObj->ExpandNode(true, ParentID);
+			GraphObj->ExpandNode(bDependency, ParentID);
 			bFoundOverflow = true;
 		}
 	}
