[Reference Viewer] ReferencedPropertiesNode is not being destroyed when running Resolve Referencing Properties multiple times

Add logic to keep track of existing Referencing Properties Node, so we can re-use them.

diff --git a/Engine/Plugins/Editor/AssetManagerEditor/Source/AssetManagerEditor/Private/ReferenceViewer/EdGraph_ReferenceViewer.cpp b/Engine/Plugins/Editor/AssetManagerEditor/Source/AssetManagerEditor/Private/ReferenceViewer/EdGraph_ReferenceViewer.cpp
--- a/Engine/Plugins/Editor/AssetManagerEditor/Source/AssetManagerEditor/Private/ReferenceViewer/EdGraph_ReferenceViewer.cpp
+++ b/Engine/Plugins/Editor/AssetManagerEditor/Source/AssetManagerEditor/Private/ReferenceViewer/EdGraph_ReferenceViewer.cpp
@@ -1455,10 +1455,43 @@ UEdGraphNode_Reference* UEdGraph_ReferenceViewer::CreateReferenceNode()
 	return Cast<UEdGraphNode_Reference>(CreateNode(UEdGraphNode_Reference::StaticClass(), bSelectNewNode));
 }
 
-UEdGraphNode_ReferencedProperties* UEdGraph_ReferenceViewer::CreateReferencedPropertiesNode()
+UEdGraphNode_ReferencedProperties* UEdGraph_ReferenceViewer::CreateReferencedPropertiesNode(
+	const TArray<FReferencingPropertyDescription>& InPropertiesDescriptionArray,
+	const TObjectPtr<UEdGraphNode_Reference>& InReferencingNode,
+	const TObjectPtr<UEdGraphNode_Reference>& InReferencedNode
+)
 {
-	constexpr  bool bSelectNewNode = false;
-	return Cast<UEdGraphNode_ReferencedProperties>(CreateNode(UEdGraphNode_ReferencedProperties::StaticClass(), bSelectNewNode));
+	uint32 NodesPairHash = GetTypeHash(InReferencingNode) ^ GetTypeHash(InReferencedNode);
+
+	UEdGraphNode_ReferencedProperties* PropertiesNode = nullptr;
+
+	if (ReferencedPropertiesNodes.Contains(NodesPairHash))
+	{
+		if (TWeakObjectPtr<UEdGraphNode_ReferencedProperties>* PropertiesNodePtr =
+			ReferencedPropertiesNodes.Find(NodesPairHash))
+		{
+			if (PropertiesNodePtr->IsValid())
+			{
+				PropertiesNode = PropertiesNodePtr->Get();
+			}
+		}
+	}
+	else
+	{
+		constexpr bool bSelectNewNode = false;
+		PropertiesNode = Cast<UEdGraphNode_ReferencedProperties>(
+			CreateNode(UEdGraphNode_ReferencedProperties::StaticClass(), bSelectNewNode)
+		);
+
+		ReferencedPropertiesNodes.Emplace(NodesPairHash, PropertiesNode);
+	}
+
+	if (PropertiesNode)
+	{
+		PropertiesNode->SetupReferencedPropertiesNode(InPropertiesDescriptionArray, InReferencingNode, InReferencedNode);
+	}
+
+	return PropertiesNode;
 }
 
 void UEdGraph_ReferenceViewer::RemoveAllNodes()
diff --git a/Engine/Plugins/Editor/AssetManagerEditor/Source/AssetManagerEditor/Private/ReferenceViewer/SReferenceViewer.cpp b/Engine/Plugins/Editor/AssetManagerEditor/Source/AssetManagerEditor/Private/ReferenceViewer/SReferenceViewer.cpp
--- a/Engine/Plugins/Editor/AssetManagerEditor/Source/AssetManagerEditor/Private/ReferenceViewer/SReferenceViewer.cpp
+++ b/Engine/Plugins/Editor/AssetManagerEditor/Source/AssetManagerEditor/Private/ReferenceViewer/SReferenceViewer.cpp
@@ -2393,10 +2393,7 @@ void SReferenceViewer::ResolveReferencingProperties() const
 					continue;
 				}
 
-				if (UEdGraphNode_ReferencedProperties* PropertiesNode = GraphObj->CreateReferencedPropertiesNode())
-				{
-					PropertiesNode->SetupReferencedPropertiesNode(ReferencingPropertiesArray, ReferencingNode, ReferencedNode);
-				}
+				GraphObj->CreateReferencedPropertiesNode(ReferencingPropertiesArray, ReferencingNode, ReferencedNode);
 			}
 		}
 	}
diff --git a/Engine/Plugins/Editor/AssetManagerEditor/Source/AssetManagerEditor/Public/ReferenceViewer/EdGraph_ReferenceViewer.h b/Engine/Plugins/Editor/AssetManagerEditor/Source/AssetManagerEditor/Public/ReferenceViewer/EdGraph_ReferenceViewer.h
--- a/Engine/Plugins/Editor/AssetManagerEditor/Source/AssetManagerEditor/Public/ReferenceViewer/EdGraph_ReferenceViewer.h
+++ b/Engine/Plugins/Editor/AssetManagerEditor/Source/AssetManagerEditor/Public/ReferenceViewer/EdGraph_ReferenceViewer.h
@@ -138,7 +138,8 @@ private:
 
 	UEdGraphNode_Reference* CreateReferenceNode();
 
-	UEdGraphNode_ReferencedProperties* CreateReferencedPropertiesNode();
+	UEdGraphNode_ReferencedProperties* CreateReferencedPropertiesNode(const TArray<FReferencingPropertyDescription>& InPropertiesDescriptionArray
+	, const TObjectPtr<UEdGraphNode_Reference>& InReferencingNode, const TObjectPtr<UEdGraphNode_Reference>& InReferencedNode);
 
 	/* Generates a NodeInfo structure then used to generate and layout the graph nodes */
 	void RecursivelyPopulateNodeInfos(bool bReferencers, const TArray<FAssetIdentifier>& Identifiers, TMap<FAssetIdentifier, FReferenceNodeInfo>& NodeInfos, int32 CurrentDepth, int32 MaxDepth);
@@ -233,5 +234,8 @@ private:
 
 	FAssetIdentifier TargetIdentifier;
 
+	/** Keeping track of existing Referencing Properties Nodes */
+	TMap<uint32, TWeakObjectPtr<UEdGraphNode_ReferencedProperties>> ReferencedPropertiesNodes;
+
 	friend SReferenceViewer;
 };
