[Reference Viewer] Multiple selection in Reference Viewer allows to Resolve Referencing Properties on filtered out classes

Unsupported nodes properties (e.g. Material Instances) won't be resolved.
For multiple selection, only supported nodes properties will be resolved, and unsupported will be ignored.

When selecting only unsupported nodes, the Resolve Properties action will not be shown.

diff --git a/Engine/Plugins/Editor/AssetManagerEditor/Source/AssetManagerEditor/Private/ReferenceViewer/ReferenceViewerSchema.cpp b/Engine/Plugins/Editor/AssetManagerEditor/Source/AssetManagerEditor/Private/ReferenceViewer/ReferenceViewerSchema.cpp
--- a/Engine/Plugins/Editor/AssetManagerEditor/Source/AssetManagerEditor/Private/ReferenceViewer/ReferenceViewerSchema.cpp
+++ b/Engine/Plugins/Editor/AssetManagerEditor/Source/AssetManagerEditor/Private/ReferenceViewer/ReferenceViewerSchema.cpp
@@ -9,8 +9,6 @@
 #include "ConnectionDrawingPolicy.h"
 #include "ToolMenuEntry.h"
 #include "ToolMenuSection.h"
-#include "Materials/MaterialFunctionInterface.h"
-#include "Materials/MaterialInstance.h"
 #include "ReferenceViewer/EdGraphNode_Reference.h"
 
 #include UE_INLINE_GENERATED_CPP_BY_NAME(ReferenceViewerSchema)
@@ -156,31 +154,12 @@ void UReferenceViewerSchema::GetContextMenuActions(UToolMenu* Menu, UGraphNodeCo
 	}
 
 	{
-		bool bResolveProperties = false;
-		if (Context)
-		{
-			if (const UEdGraphNode_Reference* ReferenceNode = Cast<UEdGraphNode_Reference>(Context->Node))
-			{
-				bResolveProperties = true;
-				FAssetData AssetData = ReferenceNode->GetAssetData();
-				if (AssetData.IsInstanceOf<UMaterialFunctionInterface>() || AssetData.IsInstanceOf<UMaterialInstance>())
-				{
-					bResolveProperties = false;
-				}
-			}
-		}
-
 		FToolMenuSection& Section = Menu->AddSection(TEXT("References"), NSLOCTEXT("ReferenceViewerSchema", "ReferencesSectionLabel", "References"));
 		Section.AddMenuEntry(FAssetManagerEditorCommands::Get().CopyReferencedObjects);
 		Section.AddMenuEntry(FAssetManagerEditorCommands::Get().CopyReferencingObjects);
 		Section.AddMenuEntry(FAssetManagerEditorCommands::Get().ShowReferencedObjects);
 		Section.AddMenuEntry(FAssetManagerEditorCommands::Get().ShowReferencingObjects);
-
-		if (bResolveProperties)
-		{
-			Section.AddMenuEntry(FAssetManagerEditorCommands::Get().ResolveReferencingProperties);
-		}
-
+		Section.AddMenuEntry(FAssetManagerEditorCommands::Get().ResolveReferencingProperties);
 		Section.AddMenuEntry(FAssetManagerEditorCommands::Get().ShowReferenceTree);
 		Section.AddMenuEntry(FAssetManagerEditorCommands::Get().ViewSizeMap);
 
