[Reference Viewer] Resolving Referencing Properties only shows the first property found

- Switching referenced object properties search from using FReferenceChainSearch to a custom FArchiveUObject class. For the moment being, we only retrieve the Property name.

Also adding a delegate which allows to refresh the Properties Node in the graph.
Currently used to properly get variable names if they are changed in Blueprints, it could be plugged in in the future so that BP recompilation or Graph refresh events always end up with an existing and updated node.

diff --git a/Engine/Plugins/Editor/AssetManagerEditor/Source/AssetManagerEditor/Private/ReferenceViewer/EdGraphNode_ReferencedProperties.cpp b/Engine/Plugins/Editor/AssetManagerEditor/Source/AssetManagerEditor/Private/ReferenceViewer/EdGraphNode_ReferencedProperties.cpp
--- a/Engine/Plugins/Editor/AssetManagerEditor/Source/AssetManagerEditor/Private/ReferenceViewer/EdGraphNode_ReferencedProperties.cpp
+++ b/Engine/Plugins/Editor/AssetManagerEditor/Source/AssetManagerEditor/Private/ReferenceViewer/EdGraphNode_ReferencedProperties.cpp
@@ -38,6 +38,11 @@ void UEdGraphNode_ReferencedProperties::SetupReferencedPropertiesNode(const TArr
 	ReferencingNode = InReferencingNode;
 	ReferencedNode = InReferencedNode;
 
+	if (OnPropertiesDescriptionUpdated().IsBound())
+	{
+		OnPropertiesDescriptionUpdated().Broadcast();
+	}
+
 	RefreshLocation();
 }
 
diff --git a/Engine/Plugins/Editor/AssetManagerEditor/Source/AssetManagerEditor/Private/ReferenceViewer/EdGraph_ReferenceViewer.cpp b/Engine/Plugins/Editor/AssetManagerEditor/Source/AssetManagerEditor/Private/ReferenceViewer/EdGraph_ReferenceViewer.cpp
--- a/Engine/Plugins/Editor/AssetManagerEditor/Source/AssetManagerEditor/Private/ReferenceViewer/EdGraph_ReferenceViewer.cpp
+++ b/Engine/Plugins/Editor/AssetManagerEditor/Source/AssetManagerEditor/Private/ReferenceViewer/EdGraph_ReferenceViewer.cpp
@@ -13,14 +13,11 @@
 #include "ICollectionManager.h"
 #include "CollectionManagerModule.h"
 #include "AssetManagerEditorModule.h"
-#include "SReferencedPropertiesNode.h"
 #include "Engine/AssetManager.h"
 #include "Engine/SCS_Node.h"
 #include "Engine/SimpleConstructionScript.h"
 #include "StructUtils/UserDefinedStruct.h"
-#include "Interfaces/IPluginManager.h"
 #include "ReferenceViewer/EdGraphNode_ReferencedProperties.h"
-#include "UObject/ReferenceChainSearch.h"
 
 #include UE_INLINE_GENERATED_CPP_BY_NAME(EdGraph_ReferenceViewer)
 
@@ -548,107 +545,81 @@ TArray<FReferencingPropertyDescription> UEdGraph_ReferenceViewer::RetrieveRefere
 	}
 
 	// The code above finds Assets when used as types (e.g. BP Enum, Struct or BPs) but not as values (e.g. a Static Mesh used as variable)
-	// To find those, we check the Reference Chain
-	TSet<FString> FoundProperties;
-	FReferenceChainSearch ReferenceChainSearch(InReferencedAsset, EReferenceChainSearchMode::Direct);
-	for (FReferenceChainSearch::FReferenceChain* ObjectRefChain : ReferenceChainSearch.GetReferenceChains())
+	// To find those, we serialize the ReferencingObject while looking for ReferencedObject referencing properties.
+
+	class FArchiveReferencingProperties : public FArchiveUObject
 	{
-		for (int32 NodeIndex = 0; NodeIndex < ObjectRefChain->Num(); ++NodeIndex)
+	public:
+		FArchiveReferencingProperties(UObject* InReferencingObject, UObject* InReferencedObject, TArray<FString>* OutReferencingProperties)
+			: ReferencingProperties(OutReferencingProperties)
+			, ReferencingObject(InReferencingObject)
+			, ReferencedObject(InReferencedObject)
 		{
-			const FReferenceChainSearch::FGraphNode* ObjectRefChainLink = ObjectRefChain->GetNode(NodeIndex);
-			UObject* LinkObject = ObjectRefChainLink->ObjectInfo->TryResolveObject();
-
-			if (!LinkObject)
-			{
-				continue;
-			}
-
-			// Let's skip the referenced asset itself
-			if (LinkObject == InReferencedAsset || LinkObject->GetOuter() == InReferencedAsset)
-			{
-				continue;
-			}
-
-			bool bReferencerFound = false;
-
-			UClass* LinkObjClass = LinkObject->GetClass();
-
-			// Check Class or Class Generated by for BPs
-			if (LinkObjClass == InReferencer->GetClass() ||
-				LinkObjClass->ClassGeneratedBy == InReferencer)
-			{
-				bReferencerFound = true;
-			}
-
-			FString IntermediateObjectName;
+			ArIsObjectReferenceCollector = true;
+			ArIgnoreOuterRef = false;
+			ArIgnoreArchetypeRef = true;
+			ArIgnoreClassGeneratedByRef = true;
+			ArIgnoreClassRef = true;
+
+			SetShouldSkipCompilingAssets(false);
+			ReferencingObjectPackage = ReferencingObject->GetPackage();
+			ReferencingObject->Serialize(*this);
+		}
 
-			if (!bReferencerFound)
+		virtual FArchive& operator<<(UObject*& InSerializedObject) override
+		{
+			if (InSerializedObject)
 			{
-				// Let's check outer as well, in case the referencing property is a field of an Object/Component
-				if (const UObject* LinkObjectOuter = LinkObject->GetOuter())
+				if (InSerializedObject == ReferencedObject)
 				{
-					// Check Class or Class Generated by for outer
-					if (LinkObjectOuter->GetClass() == InReferencer->GetClass())
-					{
-						bReferencerFound = true;
-					}
-					else if (const UClass* LinkObjectOuterClass = LinkObjectOuter->GetClass())
+					if (const FProperty* const Property = GetSerializedProperty())
 					{
-						if (LinkObjectOuterClass->ClassGeneratedBy == InReferencer)
+						if (const UObject* const PropertyOwner = Property->GetOwnerUObject())
 						{
-							bReferencerFound = true;
+							// Make sure we are only showing properties which are part of the current package
+							// Which would add no real meaningful information to the properties list
+							if (PropertyOwner->IsInPackage(ReferencingObjectPackage))
+							{
+								ReferencingProperties->AddUnique(Property->GetName());
+							}
 						}
 					}
-
-					if (bReferencerFound)
-					{
-						// Store the Object name, so we can display clearer information
-						// TODO: maybe improve the way this info is collected in FReferencingPropertyDescription
-						IntermediateObjectName = LinkObject->GetName();
-					}
 				}
-			}
-
-			if (!bReferencerFound)
-			{
-				continue;
-			}
-
-			const FReferenceChainSearch::TReferenceInfo<FReferenceChainSearch::FGraphNode>* ReferenceInfo = ObjectRefChain->GetReferenceInfo(NodeIndex);
-			if (!ReferenceInfo)
-			{
-				continue;
-			}
-
-			if (ReferenceInfo->Type != FReferenceChainSearch::EReferenceType::Property)
-			{
-				continue;
-			}
 
-			if (UObject* ReferenceObject = ReferenceInfo->Object->ObjectInfo->TryResolveObject())
-			{
-				if (ReferenceObject == InReferencedAsset)
+				if (InSerializedObject->IsInPackage(ReferencingObjectPackage))
 				{
-					FString RefName = ReferenceInfo->ReferencerName.ToString();
+					bool bAlreadyExists;
+					SerializedObjects.Add(InSerializedObject, &bAlreadyExists);
 
-					if (!IntermediateObjectName.IsEmpty())
+					if (!bAlreadyExists)
 					{
-						// TODO: can be improved by better design.
-						// This is a quick way to display useful info,
-						// e.g. MyStaticMeshComponent:StaticMesh instead of just StaticMesh
-						RefName = IntermediateObjectName + TEXT(":") + RefName;
+						InSerializedObject->Serialize(*this);
 					}
-
-					FoundProperties.Add(RefName);
 				}
 			}
+
+			return *this;
 		}
-	}
 
-	// Adding all found values refs to the list of referencing properties
-	for (const FString& PropertyName : FoundProperties)
+	private:
+		/** Stored pointer to array of objects we add object references to */
+		TArray<FString>* ReferencingProperties;
+
+		/** Tracks the objects which have been serialized by this archive, to prevent recursion */
+		TSet<UObject*> SerializedObjects;
+
+		UObject* ReferencingObject;
+		UObject* ReferencedObject;
+		UPackage* ReferencingObjectPackage;
+	};
+
+	TArray<FString> ReferencingPropertiesArray;
+	FArchiveReferencingProperties Mapper(InReferencer, InReferencedAsset, &ReferencingPropertiesArray);
+	for (const FString& Property : ReferencingPropertiesArray)
 	{
-		AddReferencingProperty(PropertyName, InReferencer->GetName(), FReferencingPropertyDescription::EAssetReferenceType::Value);
+		AddReferencingProperty(
+			Property, InReferencer->GetName(), FReferencingPropertyDescription::EAssetReferenceType::Value
+		);
 	}
 
 	return ReferencingProperties;
diff --git a/Engine/Plugins/Editor/AssetManagerEditor/Source/AssetManagerEditor/Private/ReferenceViewer/SReferencedPropertiesNode.cpp b/Engine/Plugins/Editor/AssetManagerEditor/Source/AssetManagerEditor/Private/ReferenceViewer/SReferencedPropertiesNode.cpp
--- a/Engine/Plugins/Editor/AssetManagerEditor/Source/AssetManagerEditor/Private/ReferenceViewer/SReferencedPropertiesNode.cpp
+++ b/Engine/Plugins/Editor/AssetManagerEditor/Source/AssetManagerEditor/Private/ReferenceViewer/SReferencedPropertiesNode.cpp
@@ -36,9 +36,22 @@ void SReferencedPropertiesNode::Construct(const FArguments& InArgs, UEdGraphNode
 	GraphNode = InReferencedPropertiesNode;
 	SetCursor(EMouseCursor::CardinalCross);
 
+	if (InReferencedPropertiesNode)
+	{
+		InReferencedPropertiesNode->OnPropertiesDescriptionUpdated().AddRaw(this, &SReferencedPropertiesNode::UpdateGraphNode);
+	}
+
 	UpdateGraphNode();
 }
 
+SReferencedPropertiesNode::~SReferencedPropertiesNode()
+{
+	if (UEdGraphNode_ReferencedProperties* ReferencedProperties = Cast<UEdGraphNode_ReferencedProperties>(GraphNode))
+	{
+		ReferencedProperties->OnPropertiesDescriptionUpdated().RemoveAll(this);
+	}
+}
+
 void SReferencedPropertiesNode::UpdateGraphNode()
 {
 	// No pins
diff --git a/Engine/Plugins/Editor/AssetManagerEditor/Source/AssetManagerEditor/Private/ReferenceViewer/SReferencedPropertiesNode.h b/Engine/Plugins/Editor/AssetManagerEditor/Source/AssetManagerEditor/Private/ReferenceViewer/SReferencedPropertiesNode.h
--- a/Engine/Plugins/Editor/AssetManagerEditor/Source/AssetManagerEditor/Private/ReferenceViewer/SReferencedPropertiesNode.h
+++ b/Engine/Plugins/Editor/AssetManagerEditor/Source/AssetManagerEditor/Private/ReferenceViewer/SReferencedPropertiesNode.h
@@ -23,6 +23,8 @@ public:
 
 	void Construct(const FArguments& InArgs, UEdGraphNode_ReferencedProperties* InReferencedPropertiesNode);
 
+	~SReferencedPropertiesNode();
+
 	// ~Begin SGraphNode
 	virtual void UpdateGraphNode() override;
 	virtual bool IsNodeEditable() const override { return false; }
diff --git a/Engine/Plugins/Editor/AssetManagerEditor/Source/AssetManagerEditor/Public/ReferenceViewer/EdGraphNode_ReferencedProperties.h b/Engine/Plugins/Editor/AssetManagerEditor/Source/AssetManagerEditor/Public/ReferenceViewer/EdGraphNode_ReferencedProperties.h
--- a/Engine/Plugins/Editor/AssetManagerEditor/Source/AssetManagerEditor/Public/ReferenceViewer/EdGraphNode_ReferencedProperties.h
+++ b/Engine/Plugins/Editor/AssetManagerEditor/Source/AssetManagerEditor/Public/ReferenceViewer/EdGraphNode_ReferencedProperties.h
@@ -64,7 +64,6 @@ private:
 };
 
 
-
 /**
  * A node to display a list of Node properties which are referencing another Node/Asset
  */
@@ -90,6 +89,9 @@ public:
 	 */
 	void RefreshLocation(const FVector2f& InNodeSize = FVector2f::Zero());
 
+	DECLARE_MULTICAST_DELEGATE(FOnPropertiesDescriptionUpdated)
+	FOnPropertiesDescriptionUpdated& OnPropertiesDescriptionUpdated() { return OnPropertiesDescriptionUpdatedDelegate; }
+
 private:
 	TArray<FReferencingPropertyDescription> ReferencedPropertyDescription;
 
@@ -98,4 +100,6 @@ private:
 
 	UPROPERTY()
 	TObjectPtr<UEdGraphNode_Reference> ReferencedNode;
+
+	FOnPropertiesDescriptionUpdated OnPropertiesDescriptionUpdatedDelegate;
 };
