[Reference Viewer] Reference Viewer Tab causing layout issues and potentially crashes with Tabs and Editor menus layout

One of the two Reference Viewer Tab creation calls was missing the Tab role setup, and was defaulting to Panel instead of Nomad.

diff --git a/Engine/Plugins/Editor/AssetManagerEditor/Source/AssetManagerEditor/Private/AssetManagerEditorModule.cpp b/Engine/Plugins/Editor/AssetManagerEditor/Source/AssetManagerEditor/Private/AssetManagerEditorModule.cpp
--- a/Engine/Plugins/Editor/AssetManagerEditor/Source/AssetManagerEditor/Private/AssetManagerEditorModule.cpp
+++ b/Engine/Plugins/Editor/AssetManagerEditor/Source/AssetManagerEditor/Private/AssetManagerEditorModule.cpp
@@ -742,6 +742,7 @@ TSharedRef<SDockTab> FAssetManagerEditorModule::SpawnAssetDiskSizeTab(const FSpa
 TSharedRef<SDockTab> FAssetManagerEditorModule::SpawnReferenceViewerTab(const FSpawnTabArgs& Args)
 {
 	TSharedRef<SDockTab> NewTab = SNew(SDockTab)
+		.TabRole(NomadTab)
 		.OnTabClosed_Raw(this, &FAssetManagerEditorModule::OnReferenceViewerTabClosed)
 		[
 			SAssignNew(ReferenceViewerUI, SReferenceViewer)
