[Reference Viewer] ReferencedPropertiesNode will not correctly update Property names after the variable is modified

Adds logic which will update and refresh ReferencedPropertiesNode once a referencing asset changes (e.g. a Blueprint is Saved).

diff --git a/Engine/Plugins/Editor/AssetManagerEditor/Source/AssetManagerEditor/Private/ReferenceViewer/EdGraphNode_ReferencedProperties.cpp b/Engine/Plugins/Editor/AssetManagerEditor/Source/AssetManagerEditor/Private/ReferenceViewer/EdGraphNode_ReferencedProperties.cpp
--- a/Engine/Plugins/Editor/AssetManagerEditor/Source/AssetManagerEditor/Private/ReferenceViewer/EdGraphNode_ReferencedProperties.cpp
+++ b/Engine/Plugins/Editor/AssetManagerEditor/Source/AssetManagerEditor/Private/ReferenceViewer/EdGraphNode_ReferencedProperties.cpp
@@ -30,6 +30,25 @@ FString FReferencingPropertyDescription::GetTypeAsString() const
 	return TypeString;
 }
 
+UObject* UEdGraphNode_ReferencedProperties::GetReferencingObject() const
+{
+	if (ReferencingNode)
+	{
+		return ReferencingNode->GetAssetData().GetAsset();
+	}
+
+	return nullptr;
+}
+
+UObject* UEdGraphNode_ReferencedProperties::GetReferencedObject() const
+{
+	if (ReferencedNode)
+	{
+		return ReferencedNode->GetAssetData().GetAsset();
+	}
+
+	return nullptr;
+}
 
 void UEdGraphNode_ReferencedProperties::SetupReferencedPropertiesNode(const TArray<FReferencingPropertyDescription>& InPropertiesDescription
 	, const TObjectPtr<UEdGraphNode_Reference>& InReferencingNode, const TObjectPtr<UEdGraphNode_Reference>& InReferencedNode)
diff --git a/Engine/Plugins/Editor/AssetManagerEditor/Source/AssetManagerEditor/Private/ReferenceViewer/EdGraph_ReferenceViewer.cpp b/Engine/Plugins/Editor/AssetManagerEditor/Source/AssetManagerEditor/Private/ReferenceViewer/EdGraph_ReferenceViewer.cpp
--- a/Engine/Plugins/Editor/AssetManagerEditor/Source/AssetManagerEditor/Private/ReferenceViewer/EdGraph_ReferenceViewer.cpp
+++ b/Engine/Plugins/Editor/AssetManagerEditor/Source/AssetManagerEditor/Private/ReferenceViewer/EdGraph_ReferenceViewer.cpp
@@ -344,6 +344,48 @@ UEdGraphNode_Reference* UEdGraph_ReferenceViewer::ConstructNodes(const TArray<FA
 	return RefilterGraph();
 }
 
+void UEdGraph_ReferenceViewer::RefreshReferencedPropertiesNode(const UEdGraphNode_ReferencedProperties* InNode)
+{
+	const TObjectPtr<UEdGraphNode_Reference>& ReferencingNode = InNode->GetReferencingNode();
+	if (!ReferencingNode)
+	{
+		return;
+	}
+
+	const TObjectPtr<UEdGraphNode_Reference>& ReferencedNode = InNode->GetReferencedNode();
+	if (!ReferencedNode)
+	{
+		return;
+	}
+
+	UObject* ReferencingObject = InNode->GetReferencingObject();
+	UObject* ReferencedObject = InNode->GetReferencedObject();
+	if (!ReferencingObject || !ReferencedObject)
+	{
+		return;
+	}
+
+	TArray<FReferencingPropertyDescription> ReferencingPropertiesArray =
+		RetrieveReferencingProperties(ReferencingObject, ReferencedObject);
+	if (ReferencingPropertiesArray.IsEmpty())
+	{
+		return;
+	}
+
+	CreateReferencedPropertiesNode(ReferencingPropertiesArray, ReferencingNode, ReferencedNode);
+}
+
+void UEdGraph_ReferenceViewer::RefreshReferencedPropertiesNodes()
+{
+	for (const TPair<uint32, TWeakObjectPtr<UEdGraphNode_ReferencedProperties>>& Pair : ReferencedPropertiesNodes)
+	{
+		if (UEdGraphNode_ReferencedProperties* Node = Pair.Value.Get())
+		{
+			RefreshReferencedPropertiesNode(Node);
+		}
+	}
+}
+
 TArray<FReferencingPropertyDescription> UEdGraph_ReferenceViewer::RetrieveReferencingProperties(UObject* InReferencer, UObject* InReferencedAsset)
 {
 	// This method will check InReferencer for references to InReferencedAsset.
diff --git a/Engine/Plugins/Editor/AssetManagerEditor/Source/AssetManagerEditor/Private/ReferenceViewer/SReferenceViewer.cpp b/Engine/Plugins/Editor/AssetManagerEditor/Source/AssetManagerEditor/Private/ReferenceViewer/SReferenceViewer.cpp
--- a/Engine/Plugins/Editor/AssetManagerEditor/Source/AssetManagerEditor/Private/ReferenceViewer/SReferenceViewer.cpp
+++ b/Engine/Plugins/Editor/AssetManagerEditor/Source/AssetManagerEditor/Private/ReferenceViewer/SReferenceViewer.cpp
@@ -72,6 +72,7 @@ void SReferenceViewer::Construct(const FArguments& InArgs)
 	bRebuildingFilters = false;
 	bNeedsGraphRebuild = false;
 	bNeedsGraphRefilter = false;
+	bNeedsReferencedPropertiesUpdate = false;
 	Settings = GetMutableDefault<UReferenceViewerSettings>();
 
 	// Create an action list and register commands
@@ -620,6 +621,12 @@ void SReferenceViewer::Tick( const FGeometry& AllottedGeometry, const double InC
 			GraphObj->RefilterGraph();
 		}
 	}
+
+	if (bNeedsReferencedPropertiesUpdate)
+	{
+		bNeedsReferencedPropertiesUpdate = false;
+		GraphObj->RefreshReferencedPropertiesNodes();
+	}
 }
 
 FReply SReferenceViewer::OnKeyDown(const FGeometry& MyGeometry, const FKeyEvent& InKeyEvent)
@@ -2290,6 +2297,9 @@ void SReferenceViewer::OnAssetRegistryChanged(const FAssetData& AssetData)
 {
 	// We don't do more specific checking because that data is not exposed, and it wouldn't handle newly added references anyway
 	bDirtyResults = true;
+
+	// Make sure referenced properties node are displaying updated information
+	bNeedsReferencedPropertiesUpdate = true;
 }
 
 void SReferenceViewer::OnInitialAssetRegistrySearchComplete()
diff --git a/Engine/Plugins/Editor/AssetManagerEditor/Source/AssetManagerEditor/Private/ReferenceViewer/SReferenceViewer.h b/Engine/Plugins/Editor/AssetManagerEditor/Source/AssetManagerEditor/Private/ReferenceViewer/SReferenceViewer.h
--- a/Engine/Plugins/Editor/AssetManagerEditor/Source/AssetManagerEditor/Private/ReferenceViewer/SReferenceViewer.h
+++ b/Engine/Plugins/Editor/AssetManagerEditor/Source/AssetManagerEditor/Private/ReferenceViewer/SReferenceViewer.h
@@ -287,6 +287,7 @@ private:
 	/** Used to delay graph rebuilding during spinbox slider interaction */
 	bool bNeedsGraphRebuild;
 	bool bNeedsGraphRefilter;
+	bool bNeedsReferencedPropertiesUpdate;
 	double SliderDelayLastMovedTime = 0.0;
 	double GraphRebuildSliderDelay = 0.25;
 
diff --git a/Engine/Plugins/Editor/AssetManagerEditor/Source/AssetManagerEditor/Public/ReferenceViewer/EdGraphNode_ReferencedProperties.h b/Engine/Plugins/Editor/AssetManagerEditor/Source/AssetManagerEditor/Public/ReferenceViewer/EdGraphNode_ReferencedProperties.h
--- a/Engine/Plugins/Editor/AssetManagerEditor/Source/AssetManagerEditor/Public/ReferenceViewer/EdGraphNode_ReferencedProperties.h
+++ b/Engine/Plugins/Editor/AssetManagerEditor/Source/AssetManagerEditor/Public/ReferenceViewer/EdGraphNode_ReferencedProperties.h
@@ -79,6 +79,10 @@ public:
 
 	const TObjectPtr<UEdGraphNode_Reference>& GetReferencedNode() const { return ReferencedNode; }
 
+	UObject* GetReferencingObject() const;
+	
+	UObject* GetReferencedObject() const;
+
 	/** Initialize this Node */
 	void SetupReferencedPropertiesNode(const TArray<FReferencingPropertyDescription>& InPropertiesDescription
 		, const TObjectPtr<UEdGraphNode_Reference>& InReferencingNode, const TObjectPtr<UEdGraphNode_Reference>& InReferencedNode);
diff --git a/Engine/Plugins/Editor/AssetManagerEditor/Source/AssetManagerEditor/Public/ReferenceViewer/EdGraph_ReferenceViewer.h b/Engine/Plugins/Editor/AssetManagerEditor/Source/AssetManagerEditor/Public/ReferenceViewer/EdGraph_ReferenceViewer.h
--- a/Engine/Plugins/Editor/AssetManagerEditor/Source/AssetManagerEditor/Public/ReferenceViewer/EdGraph_ReferenceViewer.h
+++ b/Engine/Plugins/Editor/AssetManagerEditor/Source/AssetManagerEditor/Public/ReferenceViewer/EdGraph_ReferenceViewer.h
@@ -118,6 +118,9 @@ public:
 	/* Returns true if the current graph has overflow nodes */
 	bool BreadthLimitExceeded() const { return bBreadthLimitReached; };
 
+	/** Refreshes the information of existing Referenced Properties Nodes */
+	void RefreshReferencedPropertiesNodes();
+
 private:
 
 	/**
@@ -185,6 +188,7 @@ private:
 	UEdGraphNode_Reference* FindPath(const FAssetIdentifier& RootId, const FAssetIdentifier& TargetId);
 	bool FindPath_Recursive(bool bInReferencers, const FAssetIdentifier& InAssetId, const FAssetIdentifier& Target, TMap<FAssetIdentifier, FReferenceNodeInfo>& InNodeInfos, TSet<FAssetIdentifier>& Visited);
 
+	void RefreshReferencedPropertiesNode(const UEdGraphNode_ReferencedProperties* InNode);
 private:
 	/** Pool for maintaining and rendering thumbnails */
 	TSharedPtr<FAssetThumbnailPool> AssetThumbnailPool;
