[Reference Viewer] Dereferencing object without check

diff --git a/Engine/Plugins/Editor/AssetManagerEditor/Source/AssetManagerEditor/Private/ReferenceViewer/EdGraph_ReferenceViewer.cpp b/Engine/Plugins/Editor/AssetManagerEditor/Source/AssetManagerEditor/Private/ReferenceViewer/EdGraph_ReferenceViewer.cpp
--- a/Engine/Plugins/Editor/AssetManagerEditor/Source/AssetManagerEditor/Private/ReferenceViewer/EdGraph_ReferenceViewer.cpp
+++ b/Engine/Plugins/Editor/AssetManagerEditor/Source/AssetManagerEditor/Private/ReferenceViewer/EdGraph_ReferenceViewer.cpp
@@ -1180,22 +1180,27 @@ UEdGraph_ReferenceViewer::RecursivelyPopulateNodeInfos(bool bInReferencers, cons
 			bIsRedirector = true;
 
 			// We are dealing with a redirector. Let's manually retrieve its Destination Object, and set up its set of nodes explicitly
-			const FName& DestinationPackageName = Redirector->DestinationObject->GetPackage()->GetFName();
-			const FAssetIdentifier DestinationAssetId = FAssetIdentifier::FromString(DestinationPackageName.ToString());
+			if (const UObject* const DestinationObject = Redirector->DestinationObject)
+			{
+				if (const UPackage* const DestinationObjectPackage = DestinationObject->GetPackage())
+				{
+					const FName& DestinationPackageName = DestinationObjectPackage->GetFName();
+					const FAssetIdentifier DestinationAssetId = FAssetIdentifier::FromString(DestinationPackageName.ToString());
 
-			FReferenceNodeInfo& DestinationReferenceNodeInfo = InNodeInfos.FindOrAdd(DestinationAssetId, FReferenceNodeInfo(DestinationAssetId, bInReferencers));
+					FReferenceNodeInfo& DestinationReferenceNodeInfo = InNodeInfos.FindOrAdd(DestinationAssetId, FReferenceNodeInfo(DestinationAssetId, bInReferencers));
 
-			// The Destination Node parent is the Redirector one
-			DestinationReferenceNodeInfo.Parents.Emplace(InAssetId);
+					// The Destination Node parent is the Redirector one
+					DestinationReferenceNodeInfo.Parents.Emplace(InAssetId);
 
-			// Remove Children from Redirector Node, and just add the Destination Node
-			InNodeInfos[InAssetId].Children.Empty();
-			InNodeInfos[InAssetId].Children.Emplace(DestinationAssetId, EDependencyPinCategory::LinkTypeHard);
-			InNodeInfos[InAssetId].bIsRedirector = true;
+					// Remove Children from Redirector Node, and just add the Destination Node
+					InNodeInfos[InAssetId].Children.Empty();
+					InNodeInfos[InAssetId].Children.Emplace(DestinationAssetId, EDependencyPinCategory::LinkTypeHard);
+					InNodeInfos[InAssetId].bIsRedirector = true;
 
-			
-			// Populate Info, without increasing current depth - we ignore the Redirector
-			RecursivelyPopulateNodeInfos(bInReferencers, { DestinationAssetId }, InNodeInfos, 0, InMaxDepth - InCurrentDepth);
+					// Populate Info, without increasing current depth - we ignore the Redirector
+					RecursivelyPopulateNodeInfos(bInReferencers, { DestinationAssetId }, InNodeInfos, 0, InMaxDepth - InCurrentDepth);
+				}
+			}
 		}
 	}
 
