[Reference Viewer] Static Mesh, Material and possibly other assets referencing properties not listed when referenced by [REDACTED_PERSON] in a Blueprint

Adds support for "indirect" references, addressing the missing referencing properties.

There are still cases in which we cannot access properties, or their details. In that case, this CL adds a simple node stating that currently that info cannot be retrieved.

Note: there are plans to address some limitations in 5.6, also targeting the support for specific types, e.g. showing Material Parameters, Functions and Texture Samplers name instead of their "raw" C++ property name.

diff --git a/Engine/Plugins/Editor/AssetManagerEditor/Source/AssetManagerEditor/Private/ReferenceViewer/EdGraph_ReferenceViewer.cpp b/Engine/Plugins/Editor/AssetManagerEditor/Source/AssetManagerEditor/Private/ReferenceViewer/EdGraph_ReferenceViewer.cpp
--- a/Engine/Plugins/Editor/AssetManagerEditor/Source/AssetManagerEditor/Private/ReferenceViewer/EdGraph_ReferenceViewer.cpp
+++ b/Engine/Plugins/Editor/AssetManagerEditor/Source/AssetManagerEditor/Private/ReferenceViewer/EdGraph_ReferenceViewer.cpp
@@ -367,10 +367,6 @@ void UEdGraph_ReferenceViewer::RefreshReferencedPropertiesNode(const UEdGraphNod
 
 	TArray<FReferencingPropertyDescription> ReferencingPropertiesArray =
 		RetrieveReferencingProperties(ReferencingObject, ReferencedObject);
-	if (ReferencingPropertiesArray.IsEmpty())
-	{
-		return;
-	}
 
 	CreateReferencedPropertiesNode(ReferencingPropertiesArray, ReferencingNode, ReferencedNode);
 }
@@ -401,10 +397,15 @@ TArray<FReferencingPropertyDescription> UEdGraph_ReferenceViewer::RetrieveRefere
 	TArray<FReferencingPropertyDescription> ReferencingProperties;
 
 	// Registering referencing properties to the output array. Property type defaults to EReferencedPropertyType::Property
-	auto AddReferencingProperty = [&ReferencingProperties, &InReferencedAsset](const FString& InPropertyName, const FString& InReferencerName
-		, FReferencingPropertyDescription::EAssetReferenceType InPropertyType = FReferencingPropertyDescription::EAssetReferenceType::Property)
-	{
-		FReferencingPropertyDescription PropertyDescription(InPropertyName, InReferencerName, InPropertyType, InReferencedAsset->GetClass());
+	auto AddReferencingProperty = [&ReferencingProperties, &InReferencedAsset](const FString& InPropertyName, const FString& InReferencerName, const FString& InReferencedNodeName,
+									  FReferencingPropertyDescription::EAssetReferenceType InPropertyType =
+										  FReferencingPropertyDescription::EAssetReferenceType::Property,
+									  bool bInIndirectReference = false
+								  )
+	{
+		FReferencingPropertyDescription PropertyDescription(
+			InPropertyName, InReferencerName, InReferencedNodeName, InPropertyType, InReferencedAsset->GetClass(), bInIndirectReference
+		);
 
 		if (!ReferencingProperties.Contains(PropertyDescription))
 		{
@@ -438,7 +439,7 @@ TArray<FReferencingPropertyDescription> UEdGraph_ReferenceViewer::RetrieveRefere
 
 			if (bMatchFound)
 			{
-				AddReferencingProperty(CurrentStructProperty->GetDisplayNameText().ToString(), InReferencer->GetName());
+				AddReferencingProperty(CurrentStructProperty->GetDisplayNameText().ToString(), InReferencer->GetName(), InReferencedAsset->GetName());
 			}
 
 			CurrentStructProperty = CurrentStructProperty->PropertyLinkNext;
@@ -476,7 +477,7 @@ TArray<FReferencingPropertyDescription> UEdGraph_ReferenceViewer::RetrieveRefere
 				if (GeneratingBlueprintObject == InReferencedAsset)
 				{
 					// The blueprint used to generate the current CDO Component Node is the same as the referenced asset: add this to output properties names
-					AddReferencingProperty(*Node->GetVariableName().ToString(), InReferencer->GetName(), FReferencingPropertyDescription::EAssetReferenceType::Component);
+					AddReferencingProperty(*Node->GetVariableName().ToString(), InReferencer->GetName(), InReferencedAsset->GetName(), FReferencingPropertyDescription::EAssetReferenceType::Component);
 				}
 			}
 		}
@@ -539,7 +540,7 @@ TArray<FReferencingPropertyDescription> UEdGraph_ReferenceViewer::RetrieveRefere
 
 					if (bAddProperty)
 					{
-						AddReferencingProperty(BPVariableDescription.VarName.ToString(), InReferencer->GetName());
+						AddReferencingProperty(BPVariableDescription.VarName.ToString(), InReferencer->GetName(), InReferencedAsset->GetName());
 					}
 				}
 				// Other
@@ -552,7 +553,7 @@ TArray<FReferencingPropertyDescription> UEdGraph_ReferenceViewer::RetrieveRefere
 						if (Object == InReferencedAsset)
 						{
 							const FString ItemIndexString = "[" + FString::FromInt(ItemIndex) + "]";
-							AddReferencingProperty(ArrayProperty->GetFName().ToString() + ItemIndexString, InReferencer->GetName());
+							AddReferencingProperty(ArrayProperty->GetFName().ToString() + ItemIndexString, InReferencer->GetName(), InReferencedAsset->GetName());
 						}
 					}
 				}
@@ -567,7 +568,7 @@ TArray<FReferencingPropertyDescription> UEdGraph_ReferenceViewer::RetrieveRefere
 
 				if (!PropertyExportString.IsEmpty() && PropertyExportString.Contains(InReferencedAsset->GetPathName()))
 				{
-					AddReferencingProperty(PropertyIt->GetFName().ToString(), InReferencer->GetName());
+					AddReferencingProperty(PropertyIt->GetFName().ToString(), InReferencer->GetName(), InReferencedAsset->GetName());
 				}
 			}
 		}
@@ -577,7 +578,7 @@ TArray<FReferencingPropertyDescription> UEdGraph_ReferenceViewer::RetrieveRefere
 			PropertyIt->GetValue_InContainer(InReferencer, &Object);
 			if (Object == InReferencedAsset)
 			{
-				AddReferencingProperty(PropertyIt->GetDisplayNameText().ToString(), InReferencer->GetName(), FReferencingPropertyDescription::EAssetReferenceType::Value);
+				AddReferencingProperty(PropertyIt->GetDisplayNameText().ToString(), InReferencer->GetName(), InReferencedAsset->GetName(), FReferencingPropertyDescription::EAssetReferenceType::Value);
 			}
 		}
 		// Other property (should handle Struct Property and fields as well)
@@ -587,7 +588,7 @@ TArray<FReferencingPropertyDescription> UEdGraph_ReferenceViewer::RetrieveRefere
 
 			if (!PropertyExportString.IsEmpty() && PropertyExportString.Contains(InReferencedAsset->GetPathName()))
 			{
-				AddReferencingProperty(PropertyIt->GetDisplayNameText().ToString(), InReferencer->GetName());
+				AddReferencingProperty(PropertyIt->GetDisplayNameText().ToString(), InReferencer->GetName(), InReferencedAsset->GetName());
 			}
 		}
 	}
@@ -598,7 +599,7 @@ TArray<FReferencingPropertyDescription> UEdGraph_ReferenceViewer::RetrieveRefere
 	class FArchiveReferencingProperties : public FArchiveUObject
 	{
 	public:
-		FArchiveReferencingProperties(UObject* InReferencingObject, UObject* InReferencedObject, TArray<FString>* OutReferencingProperties)
+		FArchiveReferencingProperties(UObject* InReferencingObject, UObject* InReferencedObject, TArray<TTuple<FString, bool>>* OutReferencingProperties)
 			: ReferencingProperties(OutReferencingProperties)
 			, ReferencingObject(InReferencingObject)
 			, ReferencedObject(InReferencedObject)
@@ -625,10 +626,32 @@ TArray<FReferencingPropertyDescription> UEdGraph_ReferenceViewer::RetrieveRefere
 						if (const UObject* const PropertyOwner = Property->GetOwnerUObject())
 						{
 							// Make sure we are only showing properties which are part of the current package
-							// Which would add no real meaningful information to the properties list
+							// This skips properties which mostly add no real meaningful information to the properties list.
+							// Some might be nice to show, which will be taken care of in the future
 							if (PropertyOwner->IsInPackage(ReferencingObjectPackage))
 							{
-								ReferencingProperties->AddUnique(Property->GetName());
+								constexpr bool bIsIndirect = false;
+								ReferencingProperties->AddUnique(TTuple<FString, bool>(Property->GetName(), bIsIndirect));
+							}
+						}
+					}
+				}
+				else if (InSerializedObject->IsInPackage(ReferencingObjectPackage))
+				{
+					// Things like a Static Mesh referenced by a BP SM Component will generate what looks like a direct reference in the
+					// graph. Let's gather those properties as well
+					for (TFieldIterator<FObjectProperty> ObjectPropertyIt(InSerializedObject->GetClass()); ObjectPropertyIt;
+						 ++ObjectPropertyIt)
+					{
+						if (UObject* ObjectReference =
+								ObjectPropertyIt->GetObjectPropertyValue_InContainer(InSerializedObject))
+						{
+							if (ObjectReference == ReferencedObject)
+							{
+								FString PropertyName = InSerializedObject->GetFName().ToString();
+
+								constexpr bool bIsIndirect = true;
+								ReferencingProperties->AddUnique(TTuple<FString, bool>(PropertyName, bIsIndirect));
 							}
 						}
 					}
@@ -651,7 +674,7 @@ TArray<FReferencingPropertyDescription> UEdGraph_ReferenceViewer::RetrieveRefere
 
 	private:
 		/** Stored pointer to array of objects we add object references to */
-		TArray<FString>* ReferencingProperties;
+		TArray<TTuple<FString, bool>>* ReferencingProperties;
 
 		/** Tracks the objects which have been serialized by this archive, to prevent recursion */
 		TSet<UObject*> SerializedObjects;
@@ -661,13 +684,15 @@ TArray<FReferencingPropertyDescription> UEdGraph_ReferenceViewer::RetrieveRefere
 		UPackage* ReferencingObjectPackage;
 	};
 
-	TArray<FString> ReferencingPropertiesArray;
+	TArray<TTuple<FString, bool>> ReferencingPropertiesArray;
 	FArchiveReferencingProperties Mapper(InReferencer, InReferencedAsset, &ReferencingPropertiesArray);
-	for (const FString& Property : ReferencingPropertiesArray)
+	for (const TTuple<FString, bool>& Property : ReferencingPropertiesArray)
 	{
+		const FString& PropertyName = Property.Get<0>();
+		const bool bIsIndirect = Property.Get<1>();
 		AddReferencingProperty(
-			Property, InReferencer->GetName(), FReferencingPropertyDescription::EAssetReferenceType::Value
-		);
+				PropertyName, InReferencer->GetName(), InReferencedAsset->GetName(), FReferencingPropertyDescription::EAssetReferenceType::Value, bIsIndirect
+			);
 	}
 
 	return ReferencingProperties;
diff --git a/Engine/Plugins/Editor/AssetManagerEditor/Source/AssetManagerEditor/Private/ReferenceViewer/SReferenceViewer.cpp b/Engine/Plugins/Editor/AssetManagerEditor/Source/AssetManagerEditor/Private/ReferenceViewer/SReferenceViewer.cpp
--- a/Engine/Plugins/Editor/AssetManagerEditor/Source/AssetManagerEditor/Private/ReferenceViewer/SReferenceViewer.cpp
+++ b/Engine/Plugins/Editor/AssetManagerEditor/Source/AssetManagerEditor/Private/ReferenceViewer/SReferenceViewer.cpp
@@ -2411,10 +2411,6 @@ void SReferenceViewer::ResolveReferencingProperties() const
 				}
 
 				TArray<FReferencingPropertyDescription> ReferencingPropertiesArray = GraphObj->RetrieveReferencingProperties(ReferencingObject, ReferencedObject);
-				if (ReferencingPropertiesArray.IsEmpty())
-				{
-					continue;
-				}
 
 				GraphObj->CreateReferencedPropertiesNode(ReferencingPropertiesArray, ReferencingNode, ReferencedNode);
 			}
diff --git a/Engine/Plugins/Editor/AssetManagerEditor/Source/AssetManagerEditor/Private/ReferenceViewer/SReferencedPropertiesNode.cpp b/Engine/Plugins/Editor/AssetManagerEditor/Source/AssetManagerEditor/Private/ReferenceViewer/SReferencedPropertiesNode.cpp
--- a/Engine/Plugins/Editor/AssetManagerEditor/Source/AssetManagerEditor/Private/ReferenceViewer/SReferencedPropertiesNode.cpp
+++ b/Engine/Plugins/Editor/AssetManagerEditor/Source/AssetManagerEditor/Private/ReferenceViewer/SReferencedPropertiesNode.cpp
@@ -9,6 +9,8 @@
 #include "Widgets/SBoxPanel.h"
 #include "Widgets/Images/SImage.h"
 
+#define LOCTEXT_NAMESPACE "ReferencedPropertyNode"
+
 /**
  * Widget representing a referencing property
  */
@@ -16,17 +18,19 @@ class SReferencedPropertyNode : public SCompoundWidget
 {
 public:
 	SLATE_BEGIN_ARGS(SReferencedPropertiesNode)
-	{
-	}
+		{
+		}
 
 	SLATE_END_ARGS()
-
 	void Construct(const FArguments& InArgs, const FReferencingPropertyDescription& InReferencingPropertyDescription);
 
 private:
 	FText GetPropertyDisplayName() const;
 	FText GetTooltipText() const;
+	FText GetIndirectReferenceTooltipText() const;
 	const FSlateBrush* GetIconBrush() const;
+	const FSlateBrush* GetIndirectReferenceIconBrush() const;
+	EVisibility GetIndirectReferenceVisibility() const;
 
 	FReferencingPropertyDescription PropertyDescription;
 };
@@ -34,7 +38,6 @@ private:
 void SReferencedPropertiesNode::Construct(const FArguments& InArgs, UEdGraphNode_ReferencedProperties* InReferencedPropertiesNode)
 {
 	GraphNode = InReferencedPropertiesNode;
-	SetCursor(EMouseCursor::CardinalCross);
 
 	if (InReferencedPropertiesNode)
 	{
@@ -80,12 +83,27 @@ void SReferencedPropertiesNode::UpdateGraphNode()
 
 	if (const UEdGraphNode_ReferencedProperties* ReferencedProperties = Cast<UEdGraphNode_ReferencedProperties>(GraphNode))
 	{
-		for (const FReferencingPropertyDescription& PropertyDescription : ReferencedProperties->GetReferencedPropertiesDescription())
+        TArray<FReferencingPropertyDescription> ReferencePropertiesDescription = ReferencedProperties->GetReferencedPropertiesDescription();
+
+		if (!ReferencePropertiesDescription.IsEmpty())
+		{
+			for (const FReferencingPropertyDescription& PropertyDescription : ReferencedProperties->GetReferencedPropertiesDescription())
+			{
+				PropertiesBox->AddSlot()
+				[
+					SNew(SReferencedPropertyNode, PropertyDescription)
+				];
+			}
+		}
+		else
 		{
 			PropertiesBox->AddSlot()
-			[
-				SNew(SReferencedPropertyNode, PropertyDescription)
-			];
+			.Padding(FMargin(6.0f, 4.0f, 6.0f, 4.0f))
+            [
+            	SNew(STextBlock)
+            	.TextStyle(FReferenceViewerStyle::Get(), "Graph.ReferencedPropertiesText")
+	            .Text(LOCTEXT("ReferencingPropertyDataUnavailable", "Impossible to retrieve at this time."))
+            ];
 		}
 	}
 }
@@ -107,6 +125,7 @@ void SReferencedPropertyNode::Construct(const FArguments& InArgs, const FReferen
 {
 	PropertyDescription = InReferencingPropertyDescription;
 
+	// clang-format off
 	ChildSlot
 	.Padding(FMargin(6.0f, 4.0f, 6.0f, 4.0f))
 	[
@@ -132,7 +151,22 @@ void SReferencedPropertyNode::Construct(const FArguments& InArgs, const FReferen
 				.Text(this, &SReferencedPropertyNode::GetPropertyDisplayName)
 			]
 		]
+		+ SHorizontalBox::Slot()
+		.AutoWidth()
+		.VAlign(VAlign_Center)
+		.Padding(FMargin(6.f, 0.f, 0.f, 0.f))
+		[
+			SNew(SBox)
+			.ToolTipText(this, &SReferencedPropertyNode::GetIndirectReferenceTooltipText)
+			.Visibility(this, &SReferencedPropertyNode::GetIndirectReferenceVisibility)
+			[
+				SNew(SImage)
+				.Image(GetIndirectReferenceIconBrush())
+				.DesiredSizeOverride(FVector2D(10.0f, 10.0f))
+			]
+		]
 	];
+	// clang-format on
 }
 
 END_SLATE_FUNCTION_BUILD_OPTIMIZATION
@@ -144,12 +178,24 @@ FText SReferencedPropertyNode::GetPropertyDisplayName() const
 
 FText SReferencedPropertyNode::GetTooltipText() const
 {
-	return FText::FromString(TEXT("Reference used as ") + PropertyDescription.GetTypeAsString());
+	FFormatNamedArguments Arguments;
+	Arguments.Add(TEXT("0"), FText::FromString(PropertyDescription.GetReferencedNodeName()));
+	Arguments.Add(TEXT("1"), FText::FromString(PropertyDescription.GetTypeAsString()));
+
+	return FText::Format(LOCTEXT("ReferenceNameTooltip", "Reference to {0} used as {1}"), Arguments);
 }
 
-const FSlateBrush* SReferencedPropertyNode::GetIconBrush() const
+FText SReferencedPropertyNode::GetIndirectReferenceTooltipText() const
 {
+	FFormatNamedArguments Arguments;
+	Arguments.Add(TEXT("0"), FText::FromString(PropertyDescription.GetName()));
+	Arguments.Add(TEXT("1"), FText::FromString(PropertyDescription.GetReferencedNodeName()));
+
+	return FText::Format(LOCTEXT("IndirectReferenceTooltip", "Indirect reference: {0} is referencing {1}"), Arguments);
+}
 
+const FSlateBrush* SReferencedPropertyNode::GetIconBrush() const
+{
 	const UClass* Class = PropertyDescription.GetPropertyClass();
 
 	const FSlateBrush* ComponentIcon;
@@ -172,3 +218,15 @@ const FSlateBrush* SReferencedPropertyNode::GetIconBrush() const
 
 	return ComponentIcon;
 }
+
+const FSlateBrush* SReferencedPropertyNode::GetIndirectReferenceIconBrush() const
+{
+	return FAppStyle::GetBrush("ReferenceViewer.IndirectReference");
+}
+
+EVisibility SReferencedPropertyNode::GetIndirectReferenceVisibility() const
+{
+	return PropertyDescription.IsIndirect() ? EVisibility::Visible : EVisibility::Collapsed;
+}
+
+#undef LOCTEXT_NAMESPACE
diff --git a/Engine/Plugins/Editor/AssetManagerEditor/Source/AssetManagerEditor/Public/ReferenceViewer/EdGraphNode_ReferencedProperties.h b/Engine/Plugins/Editor/AssetManagerEditor/Source/AssetManagerEditor/Public/ReferenceViewer/EdGraphNode_ReferencedProperties.h
--- a/Engine/Plugins/Editor/AssetManagerEditor/Source/AssetManagerEditor/Public/ReferenceViewer/EdGraphNode_ReferencedProperties.h
+++ b/Engine/Plugins/Editor/AssetManagerEditor/Source/AssetManagerEditor/Public/ReferenceViewer/EdGraphNode_ReferencedProperties.h
@@ -23,11 +23,13 @@ struct FReferencingPropertyDescription
 		None
 	};
 
-	FReferencingPropertyDescription(const FString& InName, const FString& InReferencerName, const EAssetReferenceType& InType, const UClass* InClass) :
+	FReferencingPropertyDescription(const FString& InName, const FString& InReferencerName, const FString& InReferencedNodeName, const EAssetReferenceType& InType, const UClass* InClass, bool bInIsIndirect = false) :
 		Name(InName),
 		ReferencerName(InReferencerName),
+		ReferencedNodeName(InReferencedNodeName),
 		Type(InType),
-		PropertyClass(InClass)
+		PropertyClass(InClass),
+		bIsIndirectReference(bInIsIndirect)
 	{
 	}
 
@@ -44,6 +46,9 @@ struct FReferencingPropertyDescription
 	/** Returns the name of the property referencer */
 	const FString& GetReferencerName() const { return ReferencerName; }
 
+	/** Returns the name of the referenced node */
+	const FString& GetReferencedNodeName() const { return ReferencedNodeName; }
+
 	/** Returns the property type */
 	EAssetReferenceType GetType() const { return Type; }
 
@@ -52,15 +57,19 @@ struct FReferencingPropertyDescription
 
 	const UClass* GetPropertyClass() const { return PropertyClass; }
 
+	bool IsIndirect() const { return bIsIndirectReference; }
+
 private:
 	friend class SReferencedPropertyNode;
 	FReferencingPropertyDescription() = default;
 
 	FString Name;
 	FString ReferencerName;
+	FString ReferencedNodeName;
 	EAssetReferenceType Type = EAssetReferenceType::None;
 
 	const UClass* PropertyClass = nullptr;
+	bool bIsIndirectReference = false;
 };
 
 
diff --git a/Engine/Source/Editor/EditorStyle/Private/StarshipStyle.cpp b/Engine/Source/Editor/EditorStyle/Private/StarshipStyle.cpp
--- a/Engine/Source/Editor/EditorStyle/Private/StarshipStyle.cpp
+++ b/Engine/Source/Editor/EditorStyle/Private/StarshipStyle.cpp
@@ -489,6 +489,11 @@ void FStarshipEditorStyle::FStyle::SetupGeneralStyles()
 		Set("Icons.DemoteElements", new IMAGE_BRUSH_SVG("Starship/Common/convert", Icon20x20));
 	}
 
+	// Reference Viewer Icons
+	{
+		Set("ReferenceViewer.IndirectReference", new IMAGE_BRUSH_SVG("Starship/ReferenceViewer/IndirectRef", Icon16x16));
+	}
+
 	Set("UnrealDefaultThumbnail", new IMAGE_BRUSH("Starship/Common/Unreal_DefaultThumbnail", FVector2D(256, 256)));
 
 	Set( "WarningStripe", new IMAGE_BRUSH( "Common/WarningStripe", FVector2D(20,6), FLinearColor::White, ESlateBrushTileType::Horizontal ) );
