[Reference Viewer] Add Loading asset warning for Resolve Referencing Properties action

diff --git a/Engine/Plugins/Editor/AssetManagerEditor/Source/AssetManagerEditor/Private/ReferenceViewer/SReferenceViewer.cpp b/Engine/Plugins/Editor/AssetManagerEditor/Source/AssetManagerEditor/Private/ReferenceViewer/SReferenceViewer.cpp
--- a/Engine/Plugins/Editor/AssetManagerEditor/Source/AssetManagerEditor/Private/ReferenceViewer/SReferenceViewer.cpp
+++ b/Engine/Plugins/Editor/AssetManagerEditor/Source/AssetManagerEditor/Private/ReferenceViewer/SReferenceViewer.cpp
@@ -5,6 +5,7 @@
 #include "Framework/Views/TableViewMetadata.h"
 #include "Misc/ScopedSlowTask.h"
 #include "Misc/PackageName.h"
+#include "Misc/MessageDialog.h"
 #include "Widgets/Input/SSearchBox.h"
 #include "ReferenceViewer/EdGraph_ReferenceViewer.h"
 #include "Widgets/Input/SSpinBox.h"
@@ -54,6 +55,49 @@ bool IsAssetIdentifierPassingSearchTextFilter(const FAssetIdentifier& InNode, co
 	return true;
 }
 
+EAppReturnType::Type ShowAssetsNeedsToLoadMessage(const TSet<FAssetData>& UnloadedAssetsData)
+{
+	FString UnloadedAssetsNames;
+
+	int32 Count = 0;
+	constexpr int32 MaxAssetsShown = 5;
+	for (const FAssetData& Data : UnloadedAssetsData)
+	{
+		// Don't show more than 5 entries
+		if (Count++ > MaxAssetsShown - 1)
+		{
+			break;
+		}
+
+		UnloadedAssetsNames += TEXT("\n\n") + Data.GetFullName();
+	}
+
+	if (UnloadedAssetsData.Num() > MaxAssetsShown)
+	{
+		const int32 HiddenAssets = UnloadedAssetsData.Num() - 5;
+
+		FString HiddenAssetsString = TEXT("and ") + FString::FromInt(HiddenAssets) + TEXT(" more...");
+		UnloadedAssetsNames += TEXT("\n\n") + HiddenAssetsString;
+	}
+
+	FFormatNamedArguments Args;
+	Args.Add(TEXT("UnloadedAssets"), FText::FromString(UnloadedAssetsNames));
+	static FText MessageTitle(
+		LOCTEXT("ReferencingProperties_AssetsNeedLoadingTitle", "Resolve Referencing Properties: Assets Loading")
+	);
+
+	return FMessageDialog::Open(
+		EAppMsgType::OkCancel,
+		FText::Format(
+			LOCTEXT(
+				"ReferencingProperties_AssetsNeedLoading", "The following Assets will be loaded in order to resolve referencing properties for the selected nodes: \n {UnloadedAssets}\n\n Do you wish to continue?"
+			),
+			Args
+		),
+		MessageTitle
+	);
+}
+
 SReferenceViewer::~SReferenceViewer()
 {
 	Settings->SetFindPathEnabled(false); 
@@ -2366,53 +2410,131 @@ void SReferenceViewer::ResolveReferencingProperties() const
 		return ReturnObject;
 	});
 
-	TSet<UObject*> SelectedNodes = GraphEditorPtr->GetSelectedNodes();
-	if (ensure(!SelectedNodes.IsEmpty()))
+	TSet<UObject*> SelectedNodesAsObjects = GraphEditorPtr->GetSelectedNodes();
+	if (ensure(!SelectedNodesAsObjects.IsEmpty()))
 	{
-		for (UObject* SelectedNode : SelectedNodes.Array())
+		TSet<UEdGraphNode_Reference*> SelectedNodes;
+		TSet<FAssetData> UnloadedAssetsData;
+
+		// Retrieve current Reference Nodes, and keep track of those which need to be loaded
+		for (UObject* SelectedNode : SelectedNodesAsObjects.Array())
 		{
-			UEdGraphNode_Reference* ReferencedNode = Cast<UEdGraphNode_Reference>(SelectedNode);
+			UEdGraphNode_Reference* const ReferencedNode = Cast<UEdGraphNode_Reference>(SelectedNode);
 			if (!ReferencedNode)
 			{
 				continue;
 			}
 
-			UObject* ReferencedObject = GetObjectFromNode(ReferencedNode);
-			UEdGraphPin* ReferencerPin = ReferencedNode->GetReferencerPin();
+			SelectedNodes.Add(ReferencedNode);
 
-			if (!ReferencerPin || !ReferencedObject)
+			// Look for referenced note asset, and check if it's loaded
+			const FAssetData& AssetData = ReferencedNode->GetAssetData();
+			if (!AssetData.IsAssetLoaded())
 			{
-				continue;
+				UnloadedAssetsData.Add(AssetData);
+			}
+
+			// Cycle all referencing nodes, and check if they're already loaded
+			if (const UEdGraphPin* const ReferencerPin = ReferencedNode->GetReferencerPin())
+			{
+				for (const UEdGraphPin* const ReferencedPin : ReferencerPin->LinkedTo)
+				{
+					if (ReferencedPin)
+					{
+						if (const UEdGraphNode_Reference* const ReferencingNode =
+								Cast<UEdGraphNode_Reference>(ReferencedPin->GetOwningNode()))
+						{
+							const FAssetData& ReferencerAssetData = ReferencingNode->GetAssetData();
+							if (!ReferencerAssetData.IsAssetLoaded())
+							{
+								UnloadedAssetsData.Add(ReferencerAssetData);
+							}
+						}
+					}
+				}
 			}
+		}
 
-			if (ReferencerPin->LinkedTo.IsEmpty())
+		// If assets need to be loaded in order to resolve properties, let the user know
+		if (!UnloadedAssetsData.IsEmpty())
+		{
+			const EAppReturnType::Type Ret = ShowAssetsNeedsToLoadMessage(UnloadedAssetsData);
+			if (Ret == EAppReturnType::Cancel)
+			{
+				return;
+			}
+		}
+
+		FScopedSlowTask MainResolveTask(
+			SelectedNodes.Num(),
+			LOCTEXT("ReferencingProperties_ResolveTaskDialog", "Resolving Referencing Properties for selected nodes...")
+		);
+		MainResolveTask.MakeDialog(true);
+		bool bIsCanceled = false;
+
+		for (UEdGraphNode_Reference* ReferencedNode : SelectedNodes.Array())
+		{
+			if (MainResolveTask.ShouldCancel())
+			{
+				bIsCanceled = true;
+			}
+
+			if (bIsCanceled || !ReferencedNode)
+			{
+				break;
+			}
+
+			MainResolveTask.EnterProgressFrame(
+				1.0f,
+				FText::Format(
+					LOCTEXT("ReferencingProperties_ResolveTaskDialog", "Resolving Referencing Properties for {0}"),
+					FText::FText::FromName(ReferencedNode->GetAssetData().AssetName)
+				)
+			);
+
+			UObject* const ReferencedObject = GetObjectFromNode(ReferencedNode);
+			const UEdGraphPin* const ReferencerPin = ReferencedNode->GetReferencerPin();
+
+			if (!ReferencerPin || !ReferencedObject)
 			{
 				continue;
 			}
 
-			TArray<FReferencingPropertyDescription> ReferencingProperties;
-			for (UEdGraphPin* ReferencedPin : ReferencerPin->LinkedTo)
+			const TArray<UEdGraphPin*> ReferencingPins = ReferencerPin->LinkedTo;
+			if (!ReferencingPins.IsEmpty())
 			{
-				if (!ReferencedPin)
+				TArray<FReferencingPropertyDescription> ReferencingProperties;
+				for (const UEdGraphPin* const ReferencedPin : ReferencingPins)
 				{
-					continue;
-				}
+					if (!ReferencedPin)
+					{
+						continue;
+					}
 
-				UEdGraphNode_Reference* ReferencingNode = Cast<UEdGraphNode_Reference>(ReferencedPin->GetOwningNode());
-				if (!ReferencedNode)
-				{
-					continue;
-				}
+					UEdGraphNode_Reference* const ReferencingNode =
+						Cast<UEdGraphNode_Reference>(ReferencedPin->GetOwningNode());
+					if (!ReferencingNode)
+					{
+						continue;
+					}
 
-				UObject* ReferencingObject = GetObjectFromNode(ReferencingNode);
-				if (!ReferencingObject)
-				{
-					continue;
-				}
+					if (MainResolveTask.ShouldCancel())
+					{
+						bIsCanceled = true;
+						break;
+					}
 
-				TArray<FReferencingPropertyDescription> ReferencingPropertiesArray = GraphObj->RetrieveReferencingProperties(ReferencingObject, ReferencedObject);
+					UObject* ReferencingObject = GetObjectFromNode(ReferencingNode);
+					if (!ReferencingObject)
+					{
+						continue;
+					}
 
-				GraphObj->CreateReferencedPropertiesNode(ReferencingPropertiesArray, ReferencingNode, ReferencedNode);
+					TArray<FReferencingPropertyDescription> ReferencingPropertiesArray =
+						GraphObj->RetrieveReferencingProperties(ReferencingObject, ReferencedObject);
+
+					GraphObj->CreateReferencedPropertiesNode(ReferencingPropertiesArray, ReferencingNode, ReferencedNode);
+				}
 			}
 		}
 	}
