[Reference Viewer] //UE5/Main - Build Engine Localization - Text conflict from LOCTEXT macro for namespace for SReferenceViewer.cpp

Fixing duplicate LOCTEXT macro

diff --git a/Engine/Plugins/Editor/AssetManagerEditor/Source/AssetManagerEditor/Private/ReferenceViewer/SReferenceViewer.cpp b/Engine/Plugins/Editor/AssetManagerEditor/Source/AssetManagerEditor/Private/ReferenceViewer/SReferenceViewer.cpp
--- a/Engine/Plugins/Editor/AssetManagerEditor/Source/AssetManagerEditor/Private/ReferenceViewer/SReferenceViewer.cpp
+++ b/Engine/Plugins/Editor/AssetManagerEditor/Source/AssetManagerEditor/Private/ReferenceViewer/SReferenceViewer.cpp
@@ -2487,7 +2487,7 @@ void SReferenceViewer::ResolveReferencingProperties() const
 			MainResolveTask.EnterProgressFrame(
 				1.0f,
 				FText::Format(
-					LOCTEXT("ReferencingProperties_ResolveTaskDialog", "Resolving Referencing Properties for {0}"),
+					LOCTEXT("ReferencingProperties_ResolveTaskDialogDetail", "Resolving Referencing Properties for {0}"),
 					FText::FText::FromName(ReferencedNode->GetAssetData().AssetName)
 				)
 			);
