[Reference Viewer] Empty Reference Viewer tab opens when closing undocked Reference Viewer window with multiple tabs

diff --git a/Engine/Plugins/Editor/AssetManagerEditor/Source/AssetManagerEditor/Private/AssetManagerEditorModule.cpp b/Engine/Plugins/Editor/AssetManagerEditor/Source/AssetManagerEditor/Private/AssetManagerEditorModule.cpp
--- a/Engine/Plugins/Editor/AssetManagerEditor/Source/AssetManagerEditor/Private/AssetManagerEditorModule.cpp
+++ b/Engine/Plugins/Editor/AssetManagerEditor/Source/AssetManagerEditor/Private/AssetManagerEditorModule.cpp
@@ -844,7 +844,24 @@ void FAssetManagerEditorModule::OpenReferenceViewerTab(const TArray<FAssetIdenti
 	// This means its label is still not set to an asset name
 	if (!ReferenceViewerDefaultTab.IsValid())
 	{
-		ReferenceViewerDefaultTab = FGlobalTabmanager::Get()->TryInvokeTab(ReferenceViewerTabName);
+		if (!AssetsHashToTabID.IsEmpty())
+		{
+			// Let's get the last opened tab, so we use it as reference for where new tabs should go
+			TArray<FName> TabNames;
+			AssetsHashToTabID.GenerateValueArray(TabNames);
+			TabNames.Sort(
+				[](const FName A, const FName B)
+				{
+					return A.ToString() < B.ToString();
+				}
+			);
+			ReferenceViewerDefaultTab = ReferenceViewerTabs.FindRef(TabNames.Last()).Pin();
+		}
+		else
+		{
+			// There's no Reference Viewer Tab, let's retrieve the default one
+			ReferenceViewerDefaultTab = FGlobalTabmanager::Get()->TryInvokeTab(ReferenceViewerTabName);
+		}
 	}
 
 	TSharedPtr<SDockTab> CurrentAssetTab;
@@ -864,11 +881,16 @@ void FAssetManagerEditorModule::OpenReferenceViewerTab(const TArray<FAssetIdenti
 		// Generate Tab ID 
 		TabID = FName(ReferenceViewerTabName.ToString() + "_" + FString::FromInt(++GlobalTabCount));
 
+		// Get reference Tab name, used to insert the new tab at its location
+		const FName PlacementRefTabName = ReferenceViewerDefaultTab
+											? ReferenceViewerDefaultTab->GetLayoutIdentifier().TabType
+											: ReferenceViewerTabName;
+
 		// Try to place the newly created tab next to other Reference Viewer Tabs
 		FGlobalTabmanager::Get()->InsertNewDocumentTab(
-			ReferenceViewerTabName,
+			PlacementRefTabName,
 			TabID,
-			FTabManager::FLastMajorOrNomadTab(ReferenceViewerTabName),
+			FTabManager::FLastMajorOrNomadTab(PlacementRefTabName),
 			CurrentAssetTab.ToSharedRef());
 	}
 
