[Reference Viewer] Hard reference instead of soft when referenced count more than search breadth limit

diff --git a/Engine/Plugins/Editor/AssetManagerEditor/Source/AssetManagerEditor/Private/ReferenceViewer/EdGraph_ReferenceViewer.cpp b/Engine/Plugins/Editor/AssetManagerEditor/Source/AssetManagerEditor/Private/ReferenceViewer/EdGraph_ReferenceViewer.cpp
--- a/Engine/Plugins/Editor/AssetManagerEditor/Source/AssetManagerEditor/Private/ReferenceViewer/EdGraph_ReferenceViewer.cpp
+++ b/Engine/Plugins/Editor/AssetManagerEditor/Source/AssetManagerEditor/Private/ReferenceViewer/EdGraph_ReferenceViewer.cpp
@@ -1409,6 +1409,19 @@ UEdGraphNode_Reference* UEdGraph_ReferenceViewer::RecursivelyCreateNodes(bool bI
 
 					const FAssetIdentifier& OverflowNodeAssetId = OverflowNodePair.Key;
 					OverflowNode = RecursivelyCreateNodes(bInReferencers, OverflowNodeAssetId, ChildLoc, NodeInfo.AssetId, NewNode, InNodeInfos, CurrentDepth + 1, CurrentMaxDepth);
+
+					// Make sure to keep track of pin category (e.g. soft vs hard ref)
+					const TPair<FAssetIdentifier, EDependencyPinCategory>& Pair = InNodeInfos[InAssetId].Children[Breadth];
+					FName PinCategory = ::GetName(Pair.Value);
+
+					if (bInReferencers)
+					{
+						OverflowNode->GetDependencyPin()->PinType.PinCategory = PinCategory;
+					}
+					else
+					{
+						OverflowNode->GetReferencerPin()->PinType.PinCategory = PinCategory;
+					}
 				}
 			}
 
