[Reference Viewer] Reference Viewer search doesn't filter out some assets when not submitting with Enter key

diff --git a/Engine/Plugins/Editor/AssetManagerEditor/Source/AssetManagerEditor/Private/ReferenceViewer/SReferenceViewer.cpp b/Engine/Plugins/Editor/AssetManagerEditor/Source/AssetManagerEditor/Private/ReferenceViewer/SReferenceViewer.cpp
--- a/Engine/Plugins/Editor/AssetManagerEditor/Source/AssetManagerEditor/Private/ReferenceViewer/SReferenceViewer.cpp
+++ b/Engine/Plugins/Editor/AssetManagerEditor/Source/AssetManagerEditor/Private/ReferenceViewer/SReferenceViewer.cpp
@@ -2580,6 +2580,9 @@ void SReferenceViewer::HandleOnSearchTextChanged(const FText& SearchText)
 			GraphEditorPtr->SetNodeSelection(Node, true);
 		}
 	}
+
+	// Ensure shown nodes are up-to-date with search filter
+	GraphObj->RefilterGraph();
 }
 
 void SReferenceViewer::HandleOnSearchTextCommitted(const FText& SearchText, ETextCommit::Type CommitType)
