[Reference Viewer] Reference Viewer tab's label stucks when navigating between different assets

diff --git a/Engine/Plugins/Editor/AssetManagerEditor/Source/AssetManagerEditor/Private/ReferenceViewer/SReferenceViewer.cpp b/Engine/Plugins/Editor/AssetManagerEditor/Source/AssetManagerEditor/Private/ReferenceViewer/SReferenceViewer.cpp
--- a/Engine/Plugins/Editor/AssetManagerEditor/Source/AssetManagerEditor/Private/ReferenceViewer/SReferenceViewer.cpp
+++ b/Engine/Plugins/Editor/AssetManagerEditor/Source/AssetManagerEditor/Private/ReferenceViewer/SReferenceViewer.cpp
@@ -855,14 +855,22 @@ bool SReferenceViewer::IsForwardEnabled() const
 
 void SReferenceViewer::BackClicked()
 {
+	const TArray<FAssetIdentifier>& CurrentlyVisualizedAssets = GraphObj->GetCurrentGraphRootIdentifiers();
+
 	Settings->SetFindPathEnabled(false);
 	HistoryManager.GoBack();
+
+	OnReferenceViewerSelectionChanged().Broadcast(CurrentlyVisualizedAssets, GraphObj->GetCurrentGraphRootIdentifiers());
 }
 
 void SReferenceViewer::ForwardClicked()
 {
+	const TArray<FAssetIdentifier>& CurrentlyVisualizedAssets = GraphObj->GetCurrentGraphRootIdentifiers();
+
 	Settings->SetFindPathEnabled(false);
 	HistoryManager.GoForward();
+
+	OnReferenceViewerSelectionChanged().Broadcast(CurrentlyVisualizedAssets, GraphObj->GetCurrentGraphRootIdentifiers());
 }
 
 void SReferenceViewer::RefreshClicked()
