[Reference Viewer] Resolve property UI needs overflow support

diff --git a/Engine/Plugins/Editor/AssetManagerEditor/Source/AssetManagerEditor/Private/ReferenceViewer/EdGraphNode_ReferencedProperties.cpp b/Engine/Plugins/Editor/AssetManagerEditor/Source/AssetManagerEditor/Private/ReferenceViewer/EdGraphNode_ReferencedProperties.cpp
--- a/Engine/Plugins/Editor/AssetManagerEditor/Source/AssetManagerEditor/Private/ReferenceViewer/EdGraphNode_ReferencedProperties.cpp
+++ b/Engine/Plugins/Editor/AssetManagerEditor/Source/AssetManagerEditor/Private/ReferenceViewer/EdGraphNode_ReferencedProperties.cpp
@@ -6,6 +6,13 @@
 
 #include UE_INLINE_GENERATED_CPP_BY_NAME(EdGraphNode_ReferencedProperties)
 
+TSharedRef<FReferencingPropertyDescription> FReferencingPropertyDescription::MakeSharedPropertyDescription(
+	const FReferencingPropertyDescription& InPropertyDescription
+)
+{
+	return MakeShared<FReferencingPropertyDescription>(InPropertyDescription);
+}
+
 FString FReferencingPropertyDescription::GetTypeAsString() const
 {
 	FString TypeString = "";
@@ -53,7 +60,12 @@ UObject* UEdGraphNode_ReferencedProperties::GetReferencedObject() const
 void UEdGraphNode_ReferencedProperties::SetupReferencedPropertiesNode(const TArray<FReferencingPropertyDescription>& InPropertiesDescription
 	, const TObjectPtr<UEdGraphNode_Reference>& InReferencingNode, const TObjectPtr<UEdGraphNode_Reference>& InReferencedNode)
 {
-	ReferencedPropertyDescription = InPropertiesDescription;
+	ReferencedPropertyDescription.Empty();
+	for (const FReferencingPropertyDescription& PropertyDescription : InPropertiesDescription)
+	{
+		ReferencedPropertyDescription.Add(FReferencingPropertyDescription::MakeSharedPropertyDescription(PropertyDescription));
+	}
+
 	ReferencingNode = InReferencingNode;
 	ReferencedNode = InReferencedNode;
 
diff --git a/Engine/Plugins/Editor/AssetManagerEditor/Source/AssetManagerEditor/Private/ReferenceViewer/ReferenceViewerStyle.cpp b/Engine/Plugins/Editor/AssetManagerEditor/Source/AssetManagerEditor/Private/ReferenceViewer/ReferenceViewerStyle.cpp
--- a/Engine/Plugins/Editor/AssetManagerEditor/Source/AssetManagerEditor/Private/ReferenceViewer/ReferenceViewerStyle.cpp
+++ b/Engine/Plugins/Editor/AssetManagerEditor/Source/AssetManagerEditor/Private/ReferenceViewer/ReferenceViewerStyle.cpp
@@ -117,6 +117,11 @@ FReferenceViewerStyle::FReferenceViewerStyle()
 		);
 	}
 
+	const FTableRowStyle AlternatingTableRowStyle = FTableRowStyle()
+		.SetEvenRowBackgroundBrush(FSlateNoResource())
+		.SetOddRowBackgroundBrush(FSlateNoResource());
+	Set("Graph.ReferencedPropertiesTableRow", AlternatingTableRowStyle);
+
 	FSlateStyleRegistry::RegisterSlateStyle(*this);
 }
 
diff --git a/Engine/Plugins/Editor/AssetManagerEditor/Source/AssetManagerEditor/Private/ReferenceViewer/SReferencedPropertiesNode.cpp b/Engine/Plugins/Editor/AssetManagerEditor/Source/AssetManagerEditor/Private/ReferenceViewer/SReferencedPropertiesNode.cpp
--- a/Engine/Plugins/Editor/AssetManagerEditor/Source/AssetManagerEditor/Private/ReferenceViewer/SReferencedPropertiesNode.cpp
+++ b/Engine/Plugins/Editor/AssetManagerEditor/Source/AssetManagerEditor/Private/ReferenceViewer/SReferencedPropertiesNode.cpp
@@ -11,10 +11,18 @@
 
 #define LOCTEXT_NAMESPACE "ReferencedPropertyNode"
 
+namespace UE::ReferenceViewer::Private
+{
+FText GetInvalidReferenceDescription()
+{
+	return LOCTEXT("ReferenceNameInvalidTooltip", "Invalid reference description");
+}
+}
+
 /**
  * Widget representing a referencing property
  */
-class SReferencedPropertyNode : public SCompoundWidget
+class SReferencedPropertyNode : public STableRow<FReferencingPropertyDescriptionPtr>
 {
 public:
 	SLATE_BEGIN_ARGS(SReferencedPropertiesNode)
@@ -22,7 +30,7 @@ public:
 		}
 
 	SLATE_END_ARGS()
-	void Construct(const FArguments& InArgs, const FReferencingPropertyDescription& InReferencingPropertyDescription);
+	void Construct(const FArguments& InArgs, const FReferencingPropertyDescriptionPtr& InReferencingPropertyDescription, const TSharedRef<STableViewBase>& InOwnerTableView);
 
 private:
 	FText GetPropertyDisplayName() const;
@@ -32,7 +40,7 @@ private:
 	const FSlateBrush* GetIndirectReferenceIconBrush() const;
 	EVisibility GetIndirectReferenceVisibility() const;
 
-	FReferencingPropertyDescription PropertyDescription;
+	TWeakPtr<FReferencingPropertyDescription> PropertyDescription;
 };
 
 void SReferencedPropertiesNode::Construct(const FArguments& InArgs, UEdGraphNode_ReferencedProperties* InReferencedPropertiesNode)
@@ -67,45 +75,56 @@ void SReferencedPropertiesNode::UpdateGraphNode()
 
 	ContentScale.Bind(this, &SReferencedPropertiesNode::GetContentScale);
 
-	TSharedPtr<SVerticalBox> PropertiesBox;
-
-	GetOrAddSlot(ENodeZone::Center)
-	.HAlign(HAlign_Center)
-	.VAlign(VAlign_Center)
-	[
-		SNew(SBorder)
-		.Padding(FMargin(6.0f, 4.0f, 6.0f, 4.0f))
-		.BorderImage(FReferenceViewerStyle::Get().GetBrush(TEXT("Graph.ReferencedPropertiesBrush")))
-		[
-			SAssignNew(PropertiesBox, SVerticalBox)
-		]
-	];
+	TSharedPtr<SWidget> MainWidget;
 
 	if (const UEdGraphNode_ReferencedProperties* ReferencedProperties = Cast<UEdGraphNode_ReferencedProperties>(GraphNode))
 	{
-        TArray<FReferencingPropertyDescription> ReferencePropertiesDescription = ReferencedProperties->GetReferencedPropertiesDescription();
+		ReferencingPropertiesSource = ReferencedProperties->GetReferencedPropertiesDescription();
+
+		PropertiesTreeView = SNew(SListView<FReferencingPropertyDescriptionPtr>)
+				.ConsumeMouseWheel(EConsumeMouseWheel::Always)
+				.OnGenerateRow(this, &SReferencedPropertiesNode::OnGenerateRow)
+				.SelectionMode(ESelectionMode::None)
+				.ListViewStyle(&FAppStyle::Get().GetWidgetStyle<FTableViewStyle>("SimpleListView"))
+				.Orientation(Orient_Vertical)
+				.ListItemsSource(&ReferencingPropertiesSource);
 
-		if (!ReferencePropertiesDescription.IsEmpty())
+		if (!ReferencingPropertiesSource.IsEmpty())
 		{
-			for (const FReferencingPropertyDescription& PropertyDescription : ReferencedProperties->GetReferencedPropertiesDescription())
-			{
-				PropertiesBox->AddSlot()
-				[
-					SNew(SReferencedPropertyNode, PropertyDescription)
-				];
-			}
+			MainWidget = PropertiesTreeView;
+			PropertiesTreeView->RequestListRefresh();
 		}
 		else
 		{
-			PropertiesBox->AddSlot()
-			.Padding(FMargin(6.0f, 4.0f, 6.0f, 4.0f))
-            [
+			MainWidget = 
             	SNew(STextBlock)
             	.TextStyle(FReferenceViewerStyle::Get(), "Graph.ReferencedPropertiesText")
-	            .Text(LOCTEXT("ReferencingPropertyDataUnavailable", "Impossible to retrieve at this time."))
-            ];
+	            .Text(LOCTEXT("ReferencingPropertyDataUnavailable", "Impossible to retrieve at this time."));
 		}
 	}
+	else
+	{
+		MainWidget = SNullWidget::NullWidget;
+	}
+
+	// clang-format off
+	GetOrAddSlot(ENodeZone::Center)
+	.HAlign(HAlign_Center)
+	.VAlign(VAlign_Center)
+	[
+		SNew(SBox)
+		.MaxDesiredHeight(200.0f)
+		.MinDesiredWidth(200.0f)
+		[
+			SNew(SBorder)
+			.Padding(FMargin(6.0f, 4.0f, 6.0f, 4.0f))
+			.BorderImage(FReferenceViewerStyle::Get().GetBrush(TEXT("Graph.ReferencedPropertiesBrush")))
+			[
+				MainWidget.ToSharedRef()
+			]
+		]
+	];
+	// clang-format on
 }
 
 void SReferencedPropertiesNode::Tick(const FGeometry& InAllottedGeometry, const double InCurrentTime, const float InDeltaTime)
@@ -119,10 +138,21 @@ void SReferencedPropertiesNode::Tick(const FGeometry& InAllottedGeometry, const
 	}
 }
 
+TSharedRef<ITableRow> SReferencedPropertiesNode::OnGenerateRow(
+	FReferencingPropertyDescriptionPtr ReferencingPropertyDescription, const TSharedRef<STableViewBase>& TableViewBase
+) const
+{
+	return SNew(SReferencedPropertyNode, ReferencingPropertyDescription, TableViewBase);
+}
+
 BEGIN_SLATE_FUNCTION_BUILD_OPTIMIZATION
 
-void SReferencedPropertyNode::Construct(const FArguments& InArgs, const FReferencingPropertyDescription& InReferencingPropertyDescription)
+void SReferencedPropertyNode::Construct(const FArguments& InArgs, const FReferencingPropertyDescriptionPtr& InReferencingPropertyDescription, const TSharedRef<STableViewBase>& InOwnerTableView)
 {
+	STableRow::FArguments Args = STableRow::FArguments();
+	Args.Style(&FReferenceViewerStyle::Get().GetWidgetStyle<FTableRowStyle>(TEXT("Graph.ReferencedPropertiesTableRow")));
+
+	STableRow<FReferencingPropertyDescriptionPtr>::Construct(Args, InOwnerTableView);
 	PropertyDescription = InReferencingPropertyDescription;
 
 	// clang-format off
@@ -173,48 +203,66 @@ END_SLATE_FUNCTION_BUILD_OPTIMIZATION
 
 FText SReferencedPropertyNode::GetPropertyDisplayName() const
 {
-	return FText::FromString(PropertyDescription.GetName());
+	if (TSharedPtr<FReferencingPropertyDescription> PropertyDescriptionPinned = PropertyDescription.Pin())
+	{
+		return FText::FromString(PropertyDescriptionPinned->GetName());
+	}
+
+	return FText();
 }
 
 FText SReferencedPropertyNode::GetTooltipText() const
 {
-	FFormatNamedArguments Arguments;
-	Arguments.Add(TEXT("0"), FText::FromString(PropertyDescription.GetReferencedNodeName()));
-	Arguments.Add(TEXT("1"), FText::FromString(PropertyDescription.GetTypeAsString()));
+	if (TSharedPtr<FReferencingPropertyDescription> PropertyDescriptionPinned = PropertyDescription.Pin())
+	{
+		FFormatNamedArguments Arguments;
+		Arguments.Add(TEXT("0"), FText::FromString(PropertyDescriptionPinned->GetReferencedNodeName()));
+		Arguments.Add(TEXT("1"), FText::FromString(PropertyDescriptionPinned->GetTypeAsString()));
 
-	return FText::Format(LOCTEXT("ReferenceNameTooltip", "Reference to {0} used as {1}"), Arguments);
+		return FText::Format(LOCTEXT("ReferenceNameTooltip", "Reference to {0} used as {1}"), Arguments);
+	}
+
+	return UE::ReferenceViewer::Private::GetInvalidReferenceDescription();
 }
 
 FText SReferencedPropertyNode::GetIndirectReferenceTooltipText() const
 {
-	FFormatNamedArguments Arguments;
-	Arguments.Add(TEXT("0"), FText::FromString(PropertyDescription.GetName()));
-	Arguments.Add(TEXT("1"), FText::FromString(PropertyDescription.GetReferencedNodeName()));
+	if (TSharedPtr<FReferencingPropertyDescription> PropertyDescriptionPinned = PropertyDescription.Pin())
+	{
+		FFormatNamedArguments Arguments;
+		Arguments.Add(TEXT("0"), FText::FromString(PropertyDescriptionPinned->GetName()));
+		Arguments.Add(TEXT("1"), FText::FromString(PropertyDescriptionPinned->GetReferencedNodeName()));
 
-	return FText::Format(LOCTEXT("IndirectReferenceTooltip", "Indirect reference: {0} is referencing {1}"), Arguments);
+		return FText::Format(LOCTEXT("IndirectReferenceTooltip", "Indirect reference: {0} is referencing {1}"), Arguments);
+	}
+
+	return UE::ReferenceViewer::Private::GetInvalidReferenceDescription();
 }
 
 const FSlateBrush* SReferencedPropertyNode::GetIconBrush() const
 {
-	const UClass* Class = PropertyDescription.GetPropertyClass();
+	const FSlateBrush* ComponentIcon = nullptr;
 
-	const FSlateBrush* ComponentIcon;
-
-	if (Class)
+	if (TSharedPtr<FReferencingPropertyDescription> PropertyDescriptionPinned = PropertyDescription.Pin())
 	{
-		ComponentIcon = FSlateIconFinder::FindIconBrushForClass(Class);
-	}
-	else
-	{
-		if (PropertyDescription.GetType() == FReferencingPropertyDescription::EAssetReferenceType::Component)
+		const UClass* Class = PropertyDescriptionPinned->GetPropertyClass();
+
+		if (Class)
 		{
-			ComponentIcon = FSlateIconFinder::FindIconBrushForClass(UActorComponent::StaticClass(), TEXT("SCS.Component"));
+			ComponentIcon = FSlateIconFinder::FindIconBrushForClass(Class);
 		}
 		else
 		{
-			ComponentIcon = FSlateIconFinder::FindIconBrushForClass(UObject::StaticClass());
+			if (PropertyDescriptionPinned->GetType() == FReferencingPropertyDescription::EAssetReferenceType::Component)
+			{
+				ComponentIcon = FSlateIconFinder::FindIconBrushForClass(UActorComponent::StaticClass(), TEXT("SCS.Component"));
+			}
 		}
 	}
+	else
+	{
+		ComponentIcon = FSlateIconFinder::FindIconBrushForClass(UObject::StaticClass());
+	}
 
 	return ComponentIcon;
 }
@@ -226,7 +274,11 @@ const FSlateBrush* SReferencedPropertyNode::GetIndirectReferenceIconBrush() cons
 
 EVisibility SReferencedPropertyNode::GetIndirectReferenceVisibility() const
 {
-	return PropertyDescription.IsIndirect() ? EVisibility::Visible : EVisibility::Collapsed;
+	if (TSharedPtr<FReferencingPropertyDescription> PropertyDescriptionPinned = PropertyDescription.Pin())
+	{
+		return PropertyDescriptionPinned->IsIndirect() ? EVisibility::Visible : EVisibility::Collapsed;
+	}
+	return EVisibility::Collapsed;
 }
 
 #undef LOCTEXT_NAMESPACE
diff --git a/Engine/Plugins/Editor/AssetManagerEditor/Source/AssetManagerEditor/Private/ReferenceViewer/SReferencedPropertiesNode.h b/Engine/Plugins/Editor/AssetManagerEditor/Source/AssetManagerEditor/Private/ReferenceViewer/SReferencedPropertiesNode.h
--- a/Engine/Plugins/Editor/AssetManagerEditor/Source/AssetManagerEditor/Private/ReferenceViewer/SReferencedPropertiesNode.h
+++ b/Engine/Plugins/Editor/AssetManagerEditor/Source/AssetManagerEditor/Private/ReferenceViewer/SReferencedPropertiesNode.h
@@ -2,9 +2,10 @@
 
 #pragma once
 
-#include "SGraphNode.h"
+#include "ReferenceViewer/EdGraphNode_ReferencedProperties.h"
 #include "ReferenceViewer/EdGraph_ReferenceViewer.h"
-#include "Widgets/SCompoundWidget.h"
+#include "SGraphNode.h"
+#include "Widgets/Views/SListView.h"
 
 class UEdGraphNode_ReferencedProperties;
 
@@ -37,4 +38,10 @@ public:
 	// ~Begin SWidget
 	virtual void Tick(const FGeometry& InAllottedGeometry, const double InCurrentTime, const float InDeltaTime) override;
 	// ~End SWidget
+
+private:
+	TSharedRef<ITableRow> OnGenerateRow(FReferencingPropertyDescriptionPtr ReferencingPropertyDescription, const TSharedRef<STableViewBase>& TableViewBase) const;
+
+	TArray<FReferencingPropertyDescriptionPtr> ReferencingPropertiesSource;
+	TSharedPtr<SListView<FReferencingPropertyDescriptionPtr>> PropertiesTreeView;
 };
diff --git a/Engine/Plugins/Editor/AssetManagerEditor/Source/AssetManagerEditor/Public/ReferenceViewer/EdGraphNode_ReferencedProperties.h b/Engine/Plugins/Editor/AssetManagerEditor/Source/AssetManagerEditor/Public/ReferenceViewer/EdGraphNode_ReferencedProperties.h
--- a/Engine/Plugins/Editor/AssetManagerEditor/Source/AssetManagerEditor/Public/ReferenceViewer/EdGraphNode_ReferencedProperties.h
+++ b/Engine/Plugins/Editor/AssetManagerEditor/Source/AssetManagerEditor/Public/ReferenceViewer/EdGraphNode_ReferencedProperties.h
@@ -23,6 +23,10 @@ struct FReferencingPropertyDescription
 		None
 	};
 
+	static TSharedRef<FReferencingPropertyDescription> MakeSharedPropertyDescription(
+		const FReferencingPropertyDescription& InPropertyDescription
+	);
+
 	FReferencingPropertyDescription(const FString& InName, const FString& InReferencerName, const FString& InReferencedNodeName, const EAssetReferenceType& InType, const UClass* InClass, bool bInIsIndirect = false) :
 		Name(InName),
 		ReferencerName(InReferencerName),
@@ -55,7 +59,15 @@ struct FReferencingPropertyDescription
 	/** Returns the property type as a string (useful e.g. for tooltips) */
 	FString GetTypeAsString() const;
 
-	const UClass* GetPropertyClass() const { return PropertyClass; }
+	const UClass* GetPropertyClass() const
+	{
+		if (PropertyClass.IsValid())
+		{
+			return PropertyClass.Get();
+		}
+
+		return nullptr;
+	}
 
 	bool IsIndirect() const { return bIsIndirectReference; }
 
@@ -68,10 +80,10 @@ private:
 	FString ReferencedNodeName;
 	EAssetReferenceType Type = EAssetReferenceType::None;
 
-	const UClass* PropertyClass = nullptr;
+	TWeakObjectPtr<const UClass> PropertyClass = nullptr;
 	bool bIsIndirectReference = false;
 };
-
+using FReferencingPropertyDescriptionPtr = TSharedPtr<FReferencingPropertyDescription>;
 
 /**
  * A node to display a list of Node properties which are referencing another Node/Asset
@@ -82,7 +94,10 @@ class ASSETMANAGEREDITOR_API UEdGraphNode_ReferencedProperties : public UEdGraph
 	GENERATED_BODY()
 
 public:
-	const TArray<FReferencingPropertyDescription>& GetReferencedPropertiesDescription() const { return ReferencedPropertyDescription; }
+	const TArray<FReferencingPropertyDescriptionPtr>& GetReferencedPropertiesDescription() const
+	{
+		return ReferencedPropertyDescription;
+	}
 
 	const TObjectPtr<UEdGraphNode_Reference>& GetReferencingNode() const { return ReferencingNode; }
 
@@ -106,7 +121,7 @@ public:
 	FOnPropertiesDescriptionUpdated& OnPropertiesDescriptionUpdated() { return OnPropertiesDescriptionUpdatedDelegate; }
 
 private:
-	TArray<FReferencingPropertyDescription> ReferencedPropertyDescription;
+	TArray<FReferencingPropertyDescriptionPtr> ReferencedPropertyDescription;
 
 	UPROPERTY()
 	TObjectPtr<UEdGraphNode_Reference> ReferencingNode;
