[Reference Viewer] Labels with number of assets are different when opening multiple assets in the Reference Viewer

diff --git a/Engine/Plugins/Editor/AssetManagerEditor/Source/AssetManagerEditor/Private/ReferenceViewer/SReferenceViewer.cpp b/Engine/Plugins/Editor/AssetManagerEditor/Source/AssetManagerEditor/Private/ReferenceViewer/SReferenceViewer.cpp
--- a/Engine/Plugins/Editor/AssetManagerEditor/Source/AssetManagerEditor/Private/ReferenceViewer/SReferenceViewer.cpp
+++ b/Engine/Plugins/Editor/AssetManagerEditor/Source/AssetManagerEditor/Private/ReferenceViewer/SReferenceViewer.cpp
@@ -921,7 +921,11 @@ FText SReferenceViewer::GetAddressBarText() const
 			}
 			else if (CurrentGraphRootPackageNames.Num() > 1)
 			{
-				return FText::Format(LOCTEXT("AddressBarMultiplePackagesText", "{0} and {1} others"), FText::FromString(CurrentGraphRootPackageNames[0].ToString()), FText::AsNumber(CurrentGraphRootPackageNames.Num()));
+				return FText::Format(
+					LOCTEXT("AddressBarMultiplePackagesText", "{0} and {1} others"),
+					FText::FromString(CurrentGraphRootPackageNames[0].ToString()),
+					FText::AsNumber(CurrentGraphRootPackageNames.Num() - 1)
+				);
 			}
 			else
 			{
