[Reference Viewer] Add a close button to resolve Property nodes

Adds:
- title and tooltip for the node
- close button

diff --git a/Engine/Plugins/Editor/AssetManagerEditor/Source/AssetManagerEditor/Private/ReferenceViewer/EdGraph_ReferenceViewer.cpp b/Engine/Plugins/Editor/AssetManagerEditor/Source/AssetManagerEditor/Private/ReferenceViewer/EdGraph_ReferenceViewer.cpp
--- a/Engine/Plugins/Editor/AssetManagerEditor/Source/AssetManagerEditor/Private/ReferenceViewer/EdGraph_ReferenceViewer.cpp
+++ b/Engine/Plugins/Editor/AssetManagerEditor/Source/AssetManagerEditor/Private/ReferenceViewer/EdGraph_ReferenceViewer.cpp
@@ -371,6 +371,21 @@ void UEdGraph_ReferenceViewer::RefreshReferencedPropertiesNode(const UEdGraphNod
 	CreateReferencedPropertiesNode(ReferencingPropertiesArray, ReferencingNode, ReferencedNode);
 }
 
+void UEdGraph_ReferenceViewer::CloseReferencedPropertiesNode(UEdGraphNode_ReferencedProperties* InNode)
+{
+	if (InNode)
+	{
+		uint32 NodesPairHash = GetTypeHash(InNode->GetReferencingNode()) ^ GetTypeHash(InNode->GetReferencedNode());
+
+		if (ReferencedPropertiesNodes.Find(NodesPairHash))
+		{
+			ReferencedPropertiesNodes.Remove(NodesPairHash);
+		}
+
+		RemoveNode(InNode);
+	}
+}
+
 void UEdGraph_ReferenceViewer::RefreshReferencedPropertiesNodes()
 {
 	for (const TPair<uint32, TWeakObjectPtr<UEdGraphNode_ReferencedProperties>>& Pair : ReferencedPropertiesNodes)
diff --git a/Engine/Plugins/Editor/AssetManagerEditor/Source/AssetManagerEditor/Private/ReferenceViewer/ReferenceViewerStyle.cpp b/Engine/Plugins/Editor/AssetManagerEditor/Source/AssetManagerEditor/Private/ReferenceViewer/ReferenceViewerStyle.cpp
--- a/Engine/Plugins/Editor/AssetManagerEditor/Source/AssetManagerEditor/Private/ReferenceViewer/ReferenceViewerStyle.cpp
+++ b/Engine/Plugins/Editor/AssetManagerEditor/Source/AssetManagerEditor/Private/ReferenceViewer/ReferenceViewerStyle.cpp
@@ -17,6 +17,8 @@ FName FReferenceViewerStyle::StyleName("ReferenceViewerStyle");
 FReferenceViewerStyle::FReferenceViewerStyle()
 	: FSlateStyleSet(StyleName)
 {
+	const FVector2D Icon16x16(16.0f, 16.0f);
+
 	SetParentStyleName("EditorStyle");
 
 	SetContentRoot(FPaths::EnginePluginsDir() / TEXT("Editor/AssetManagerEditor/Content"));
@@ -115,10 +117,22 @@ FReferenceViewerStyle::FReferenceViewerStyle()
 		);
 	}
 
-	const FTableRowStyle AlternatingTableRowStyle = FTableRowStyle()
-		.SetEvenRowBackgroundBrush(FSlateNoResource())
-		.SetOddRowBackgroundBrush(FSlateNoResource());
-	Set("Graph.ReferencedPropertiesTableRow", AlternatingTableRowStyle);
+	// Referenced Properties List Rows
+	{
+		const FTableRowStyle AlternatingTableRowStyle = FTableRowStyle()
+		   .SetEvenRowBackgroundBrush(FSlateNoResource())
+		   .SetOddRowBackgroundBrush(FSlateNoResource());
+		Set("Graph.ReferencedPropertiesTableRow", AlternatingTableRowStyle);
+	}
+
+	// Referenced Properties Close Button
+	{
+		const FButtonStyle CloseButton = FButtonStyle()
+		   .SetNormal(CORE_IMAGE_BRUSH_SVG("Starship/Common/close-small", Icon16x16, FStyleColors::Foreground))
+		   .SetPressed(CORE_IMAGE_BRUSH_SVG("Starship/Common/close-small", Icon16x16, FStyleColors::Foreground))
+		   .SetHovered(CORE_IMAGE_BRUSH_SVG("Starship/Common/close-small", Icon16x16, FStyleColors::ForegroundHover));
+		Set("Graph.ReferencedPropertiesCloseButton", CloseButton);
+	}
 
 	FSlateStyleRegistry::RegisterSlateStyle(*this);
 }
diff --git a/Engine/Plugins/Editor/AssetManagerEditor/Source/AssetManagerEditor/Private/ReferenceViewer/SReferencedPropertiesNode.cpp b/Engine/Plugins/Editor/AssetManagerEditor/Source/AssetManagerEditor/Private/ReferenceViewer/SReferencedPropertiesNode.cpp
--- a/Engine/Plugins/Editor/AssetManagerEditor/Source/AssetManagerEditor/Private/ReferenceViewer/SReferencedPropertiesNode.cpp
+++ b/Engine/Plugins/Editor/AssetManagerEditor/Source/AssetManagerEditor/Private/ReferenceViewer/SReferencedPropertiesNode.cpp
@@ -1,13 +1,16 @@
 // Copyright Epic Games, Inc. All Rights Reserved.
 
 #include "SReferencedPropertiesNode.h"
+#include "ReferenceViewer/EdGraphNode_Reference.h"
 #include "ReferenceViewer/EdGraphNode_ReferencedProperties.h"
 #include "ReferenceViewerStyle.h"
 #include "SlateOptMacros.h"
 #include "Styling/SlateIconFinder.h"
+#include "Widgets/Images/SImage.h"
+#include "Widgets/Input/SButton.h"
 #include "Widgets/Layout/SBox.h"
+#include "Widgets/Layout/SSpacer.h"
 #include "Widgets/SBoxPanel.h"
-#include "Widgets/Images/SImage.h"
 
 #define LOCTEXT_NAMESPACE "ReferencedPropertyNode"
 
@@ -25,7 +28,7 @@ FText GetInvalidReferenceDescription()
 class SReferencedPropertyNode : public STableRow<FReferencingPropertyDescriptionPtr>
 {
 public:
-	SLATE_BEGIN_ARGS(SReferencedPropertiesNode)
+	SLATE_BEGIN_ARGS(SReferencedPropertyNode)
 		{
 		}
 
@@ -107,20 +110,71 @@ void SReferencedPropertiesNode::UpdateGraphNode()
 		MainWidget = SNullWidget::NullWidget;
 	}
 
+	const FButtonStyle* CloseButtonStyle =
+		&FReferenceViewerStyle::Get().GetWidgetStyle<FButtonStyle>(TEXT("Graph.ReferencedPropertiesCloseButton"));
+
 	// clang-format off
 	GetOrAddSlot(ENodeZone::Center)
 	.HAlign(HAlign_Center)
 	.VAlign(VAlign_Center)
 	[
-		SNew(SBox)
-		.MaxDesiredHeight(200.0f)
-		.MinDesiredWidth(200.0f)
+		SNew(SBorder)
+		.ToolTipText(this, &SReferencedPropertiesNode::GetTooltipText)
+		.ColorAndOpacity(FLinearColor::White)
+		.BorderImage(FReferenceViewerStyle::Get().GetBrush("Graph.Node.BodyBackground"))
+		.Padding(0.0f)
 		[
 			SNew(SBorder)
-			.Padding(FMargin(6.0f, 4.0f, 6.0f, 4.0f))
-			.BorderImage(FReferenceViewerStyle::Get().GetBrush(TEXT("Graph.ReferencedPropertiesBrush")))
+			.BorderImage(FReferenceViewerStyle::Get().GetBrush("Graph.Node.BodyBorder"))
+			.BorderBackgroundColor(FLinearColor(0.5f, 0.5f, 0.5f, 0.4f))
+			.Padding(0.0f)
 			[
-				MainWidget.ToSharedRef()
+				SNew(SBorder)
+				.BorderImage(FReferenceViewerStyle::Get().GetBrush("Graph.Node.Body"))
+				.Padding(0.0f)
+				[
+					SNew(SVerticalBox)
+					+ SVerticalBox::Slot()
+					.AutoHeight()
+					[
+						SNew(SBox)
+						.Padding(4.0f)
+						[
+							SNew(SHorizontalBox)
+							+ SHorizontalBox::Slot()
+							.HAlign(HAlign_Left)
+							.AutoWidth()
+							.Padding(4.0, 2.0f, 0.0f, 0.0f)
+							[
+								SNew(STextBlock)
+								.TextStyle(&FReferenceViewerStyle::Get(), TEXT("Graph.Node.NodeTitleExtraLines"))
+								.Text(LOCTEXT("ReferencingPropertiesLabel", "Referencing Properties"))
+							]
+							+ SHorizontalBox::Slot()
+							.HAlign(HAlign_Right)
+							[
+								SNew(SButton)
+								.ContentPadding(1.0f)
+								.ButtonStyle(CloseButtonStyle)
+								.OnClicked(this, &SReferencedPropertiesNode::CloseNode)
+								[
+									SNew(SSpacer)
+									.Size(CloseButtonStyle->Normal.ImageSize)
+								]
+							]
+						]
+					]
+					+ SVerticalBox::Slot()
+					[
+						SNew(SBox)
+						.MaxDesiredHeight(200.0f)
+						.MinDesiredWidth(200.0f)
+						.Padding(2.0f)
+						[
+							MainWidget.ToSharedRef()
+						]
+					]
+				]
 			]
 		]
 	];
@@ -145,6 +199,39 @@ TSharedRef<ITableRow> SReferencedPropertiesNode::OnGenerateRow(
 	return SNew(SReferencedPropertyNode, ReferencingPropertyDescription, TableViewBase);
 }
 
+FReply SReferencedPropertiesNode::CloseNode() const
+{
+	if (UEdGraphNode_ReferencedProperties* ReferencedPropertiesNode = Cast<UEdGraphNode_ReferencedProperties>(GraphNode))
+	{
+		if (UEdGraph_ReferenceViewer* Graph = Cast<UEdGraph_ReferenceViewer>(GraphNode->GetGraph()))
+		{
+			Graph->CloseReferencedPropertiesNode(ReferencedPropertiesNode);
+		}
+	}
+
+	return FReply::Handled();
+}
+
+FText SReferencedPropertiesNode::GetTooltipText() const
+{
+	if (UEdGraphNode_ReferencedProperties* ReferencedPropertiesNode = Cast<UEdGraphNode_ReferencedProperties>(GraphNode))
+	{
+		TObjectPtr<UEdGraphNode_Reference> ReferencedNode = ReferencedPropertiesNode->GetReferencedNode();
+		TObjectPtr<UEdGraphNode_Reference> ReferencingNode = ReferencedPropertiesNode->GetReferencingNode();
+
+		if (ReferencedNode && ReferencingNode)
+		{
+			FFormatNamedArguments Arguments;
+			Arguments.Add(TEXT("0"), FText::FromName(ReferencingNode->GetAssetData().AssetName));
+			Arguments.Add(TEXT("1"), FText::FromName(ReferencedNode->GetAssetData().AssetName));
+
+			return FText::Format(LOCTEXT("ReferencingPropertiesNodeTooltip", "{0} properties referencing {1}"), Arguments);
+		}
+	}
+
+	return FText();
+}
+
 BEGIN_SLATE_FUNCTION_BUILD_OPTIMIZATION
 
 void SReferencedPropertyNode::Construct(const FArguments& InArgs, const FReferencingPropertyDescriptionPtr& InReferencingPropertyDescription, const TSharedRef<STableViewBase>& InOwnerTableView)
diff --git a/Engine/Plugins/Editor/AssetManagerEditor/Source/AssetManagerEditor/Private/ReferenceViewer/SReferencedPropertiesNode.h b/Engine/Plugins/Editor/AssetManagerEditor/Source/AssetManagerEditor/Private/ReferenceViewer/SReferencedPropertiesNode.h
--- a/Engine/Plugins/Editor/AssetManagerEditor/Source/AssetManagerEditor/Private/ReferenceViewer/SReferencedPropertiesNode.h
+++ b/Engine/Plugins/Editor/AssetManagerEditor/Source/AssetManagerEditor/Private/ReferenceViewer/SReferencedPropertiesNode.h
@@ -41,6 +41,8 @@ public:
 
 private:
 	TSharedRef<ITableRow> OnGenerateRow(FReferencingPropertyDescriptionPtr ReferencingPropertyDescription, const TSharedRef<STableViewBase>& TableViewBase) const;
+	FReply CloseNode() const;
+	FText GetTooltipText() const;
 
 	TArray<FReferencingPropertyDescriptionPtr> ReferencingPropertiesSource;
 	TSharedPtr<SListView<FReferencingPropertyDescriptionPtr>> PropertiesTreeView;
diff --git a/Engine/Plugins/Editor/AssetManagerEditor/Source/AssetManagerEditor/Public/ReferenceViewer/EdGraph_ReferenceViewer.h b/Engine/Plugins/Editor/AssetManagerEditor/Source/AssetManagerEditor/Public/ReferenceViewer/EdGraph_ReferenceViewer.h
--- a/Engine/Plugins/Editor/AssetManagerEditor/Source/AssetManagerEditor/Public/ReferenceViewer/EdGraph_ReferenceViewer.h
+++ b/Engine/Plugins/Editor/AssetManagerEditor/Source/AssetManagerEditor/Public/ReferenceViewer/EdGraph_ReferenceViewer.h
@@ -121,6 +121,9 @@ public:
 	/** Refreshes the information of existing Referenced Properties Nodes */
 	void RefreshReferencedPropertiesNodes();
 
+	/** Closes (removes) the specified Referenced Properties Node */
+	void CloseReferencedPropertiesNode(UEdGraphNode_ReferencedProperties* InNode);
+
 private:
 
 	/**
