[Reference Viewer] Reference Viewer tab's label stucks when navigating between different assets

Simply removing refs used when getting current graph selection array. They where inadvertently breaking the actual fix.

diff --git a/Engine/Plugins/Editor/AssetManagerEditor/Source/AssetManagerEditor/Private/ReferenceViewer/SReferenceViewer.cpp b/Engine/Plugins/Editor/AssetManagerEditor/Source/AssetManagerEditor/Private/ReferenceViewer/SReferenceViewer.cpp
--- a/Engine/Plugins/Editor/AssetManagerEditor/Source/AssetManagerEditor/Private/ReferenceViewer/SReferenceViewer.cpp
+++ b/Engine/Plugins/Editor/AssetManagerEditor/Source/AssetManagerEditor/Private/ReferenceViewer/SReferenceViewer.cpp
@@ -855,7 +855,8 @@ bool SReferenceViewer::IsForwardEnabled() const
 
 void SReferenceViewer::BackClicked()
 {
-	const TArray<FAssetIdentifier>& CurrentlyVisualizedAssets = GraphObj->GetCurrentGraphRootIdentifiers();
+	// Storing assets visualized before using history to go back. Cannot get this as ref since the original array is about to change.
+	const TArray<FAssetIdentifier> CurrentlyVisualizedAssets = GraphObj->GetCurrentGraphRootIdentifiers();
 
 	Settings->SetFindPathEnabled(false);
 	HistoryManager.GoBack();
@@ -865,7 +866,8 @@ void SReferenceViewer::BackClicked()
 
 void SReferenceViewer::ForwardClicked()
 {
-	const TArray<FAssetIdentifier>& CurrentlyVisualizedAssets = GraphObj->GetCurrentGraphRootIdentifiers();
+	// Storing assets visualized before using history to go forward. Cannot get this as ref since the original array is about to change.
+	const TArray<FAssetIdentifier> CurrentlyVisualizedAssets = GraphObj->GetCurrentGraphRootIdentifiers();
 
 	Settings->SetFindPathEnabled(false);
 	HistoryManager.GoForward();
