Fixed typo whereby FAssetManagerEditorModule::WriteCollection tries to add PackageNames to the wrong collection share type.

diff --git a/Engine/Plugins/Editor/AssetManagerEditor/Source/AssetManagerEditor/Private/AssetManagerEditorModule.cpp b/Engine/Plugins/Editor/AssetManagerEditor/Source/AssetManagerEditor/Private/AssetManagerEditorModule.cpp
--- a/Engine/Plugins/Editor/AssetManagerEditor/Source/AssetManagerEditor/Private/AssetManagerEditorModule.cpp
+++ b/Engine/Plugins/Editor/AssetManagerEditor/Source/AssetManagerEditor/Private/AssetManagerEditorModule.cpp
@@ -2832,7 +2832,7 @@ bool FAssetManagerEditorModule::WriteCollection(ICollectionContainer& Collection
 	else if (CreateOrEmptyCollection(CollectionContainer, CollectionName, ShareType, Error))
 	{
 		if (CollectionContainer.AddToCollection(
-				CollectionName, ECollectionShareType::CST_Local, ObjectPathsToAddToCollection.Array(), nullptr, &Error))
+				CollectionName, ShareType, ObjectPathsToAddToCollection.Array(), nullptr, &Error))
 		{
 			UE_LOG(LogAssetManagerEditor, Log, TEXT("Updated collection %s"), *CollectionName.ToString());
 			ResultsMessage = FText::Format(LOCTEXT("CreateCollectionSucceeded", "Updated collection {0}"), FText::FromName(CollectionName));
