Use the platform format when copying disk file paths to the clipboard.

diff --git a/Engine/Plugins/Editor/AssetManagerEditor/Source/AssetManagerEditor/Private/AssetManagerEditorModule.cpp b/Engine/Plugins/Editor/AssetManagerEditor/Source/AssetManagerEditor/Private/AssetManagerEditorModule.cpp
--- a/Engine/Plugins/Editor/AssetManagerEditor/Source/AssetManagerEditor/Private/AssetManagerEditorModule.cpp
+++ b/Engine/Plugins/Editor/AssetManagerEditor/Source/AssetManagerEditor/Private/AssetManagerEditorModule.cpp
@@ -2625,11 +2625,13 @@ void FAssetManagerEditorModule::CopyPackagePaths(const TArray<FAssetData> Select
 
 	FString ClipboardText = FString::JoinBy(SortedItems, LINE_TERMINATOR, [](const FAssetData& Item)
 		{
-			const FString ItemFilename = FPackageName::LongPackageNameToFilename(Item.PackageName.ToString(), FPackageName::GetAssetPackageExtension());
+			FString ItemFilename = FPackageName::LongPackageNameToFilename(Item.PackageName.ToString(), FPackageName::GetAssetPackageExtension());
 
 			if (FPaths::FileExists(ItemFilename))
 			{
-				return FPaths::ConvertRelativePathToFull(ItemFilename);
+				ItemFilename = FPaths::ConvertRelativePathToFull(ItemFilename);
+				FPaths::MakePlatformFilename(ItemFilename);
+				return ItemFilename;
 			}
 			else
 			{
diff --git a/Engine/Source/Editor/ContentBrowser/Private/ContentBrowserUtils.cpp b/Engine/Source/Editor/ContentBrowser/Private/ContentBrowserUtils.cpp
--- a/Engine/Source/Editor/ContentBrowser/Private/ContentBrowserUtils.cpp
+++ b/Engine/Source/Editor/ContentBrowser/Private/ContentBrowserUtils.cpp
@@ -463,7 +463,9 @@ void ContentBrowserUtils::CopyFilePathsToClipboard(const TArray<FContentBrowserI
 		FString ItemFilename;
 		if (Item.GetItemPhysicalPath(ItemFilename) && FPaths::FileExists(ItemFilename))
 		{
-			ClipboardText += FPaths::ConvertRelativePathToFull(ItemFilename);
+			ItemFilename = FPaths::ConvertRelativePathToFull(ItemFilename);
+			FPaths::MakePlatformFilename(ItemFilename);
+			ClipboardText += ItemFilename;
 		}
 		else
 		{
diff --git a/Engine/Source/Editor/LevelEditor/Private/LevelEditorActions.cpp b/Engine/Source/Editor/LevelEditor/Private/LevelEditorActions.cpp
--- a/Engine/Source/Editor/LevelEditor/Private/LevelEditorActions.cpp
+++ b/Engine/Source/Editor/LevelEditor/Private/LevelEditorActions.cpp
@@ -3105,7 +3105,7 @@ void FLevelEditorActionCallbacks::CopyActorFilePathtoClipboard_Clicked()
 		}
 
 		const UPackage* SceneOutlinerPackage = Actor->GetSceneOutlinerTopParentPackage();
-		const FString LocalFullPath(SceneOutlinerPackage->GetLoadedPath().GetLocalFullPath());
+		FString LocalFullPath(SceneOutlinerPackage->GetLoadedPath().GetLocalFullPath());
 
 		if (SelectedActors.Num() > 1)
 		{
@@ -3125,7 +3125,10 @@ void FLevelEditorActionCallbacks::CopyActorFilePathtoClipboard_Clicked()
 			}
 		}
 
-		Result.Append(FPaths::ConvertRelativePathToFull(LocalFullPath));
+		LocalFullPath = FPaths::ConvertRelativePathToFull(LocalFullPath);
+		FPaths::MakePlatformFilename(LocalFullPath);
+
+		Result.Append(LocalFullPath);
 	}
 
 	if (Result.Len())
diff --git a/Engine/Source/Editor/SceneOutliner/Private/ActorDescTreeItem.cpp b/Engine/Source/Editor/SceneOutliner/Private/ActorDescTreeItem.cpp
--- a/Engine/Source/Editor/SceneOutliner/Private/ActorDescTreeItem.cpp
+++ b/Engine/Source/Editor/SceneOutliner/Private/ActorDescTreeItem.cpp
@@ -312,6 +312,7 @@ void FActorDescTreeItem::CopyActorFilePathtoClipboard() const
 		if (FPackageName::TryConvertLongPackageNameToFilename(ActorDescInstance->GetActorPackage().ToString(), PackageFilename, FPackageName::GetAssetPackageExtension()))
 		{
 			FString Result = FPaths::ConvertRelativePathToFull(PackageFilename);
+			FPaths::MakePlatformFilename(Result);
 			FPlatformApplicationMisc::ClipboardCopy(*Result);
 		}
 	}
diff --git a/Engine/Source/Runtime/Engine/Private/ExternalPackageHelper.cpp b/Engine/Source/Runtime/Engine/Private/ExternalPackageHelper.cpp
--- a/Engine/Source/Runtime/Engine/Private/ExternalPackageHelper.cpp
+++ b/Engine/Source/Runtime/Engine/Private/ExternalPackageHelper.cpp
@@ -266,7 +266,11 @@ void FExternalPackageHelper::CopyObjectsExternalPackageFilePathToClipboard(const
 	TArray<FString> PackageFilePaths = GetObjectsExternalPackageFilePath(InObjects);
 	if (!PackageFilePaths.IsEmpty())
 	{
-		FString Result = FString::Join(PackageFilePaths, TEXT("\n"));
+		for (FString& PackageFilePath : PackageFilePaths)
+		{
+			FPaths::MakePlatformFilename(PackageFilePath);
+		}
+		FString Result = FString::Join(PackageFilePaths, LINE_TERMINATOR);
 		check(Result.Len());
 		FPlatformApplicationMisc::ClipboardCopy(*Result);
 	}
