Adding checks to prevent recursing into the same plugin multiple times when pulling its dependencies.

Adding assets to plugins to make them quicker to iterate per plugin.

#rnx

diff --git a/Engine/Plugins/Editor/AssetManagerEditor/Source/AssetManagerEditor/Private/TreeView/AssetTable.cpp b/Engine/Plugins/Editor/AssetManagerEditor/Source/AssetManagerEditor/Private/TreeView/AssetTable.cpp
--- a/Engine/Plugins/Editor/AssetManagerEditor/Source/AssetManagerEditor/Private/TreeView/AssetTable.cpp
+++ b/Engine/Plugins/Editor/AssetManagerEditor/Source/AssetManagerEditor/Private/TreeView/AssetTable.cpp
@@ -504,14 +504,21 @@ void FAssetTable::EnumeratePluginDependencies(const FAssetTablePluginInfo& InPlu
 
 void FAssetTable::EnumerateAssetsForPlugin(const FAssetTablePluginInfo& InPlugin, TFunction<void(int32 AssetIndex)> Callback) const
 {
-	const int32 AssetCount = Assets.Num();
-	for (int32 AssetIndex = 0; AssetIndex < AssetCount; ++AssetIndex)
+	TSet<int32> PluginAssets = InPlugin.GetPluginAssets();
+	for (int32 AssetIndex : PluginAssets)
 	{
-		// The StringStore makes this comparison safe
-		if (Assets[AssetIndex].GetPluginName() == InPlugin.PluginName)
-		{
-			Callback(AssetIndex);
-		}
+		Callback(AssetIndex);
+	}
+}
+
+////////////////////////////////////////////////////////////////////////////////////////////////////
+
+void FAssetTable::PopulateAssetsInPlugins()
+{
+	for (int32 AssetIndex = 0 ; AssetIndex < Assets.Num() ; AssetIndex++)
+	{
+		FAssetTablePluginInfo& PluginInfo = GetOrCreatePluginInfo(Assets[AssetIndex].GetPluginName());
+		PluginInfo.AddPluginAsset(AssetIndex);
 	}
 }
 
@@ -1921,18 +1928,20 @@ void FAssetTablePluginInfo::ComputeDependencySizes(const FAssetTable& OwningTabl
 
 	TArray<int32> DependencyStack = PluginDependencies;
 	TSet<int32> AllDependencies;
-
+	bool bAlreadyExists = false;
 	// Gather all the dependencies
 	while (DependencyStack.Num() > 0)
 	{
 		int32 CurrentIndex = DependencyStack.Pop();
-		AllDependencies.Add(CurrentIndex);
-		// We'll end up processing some nodes twice but the real output here is AllDependencies (the set) not DependencyStack so it's ok
-		DependencyStack.Append(OwningTable.GetPluginInfoByIndex(CurrentIndex).GetDependencies());
+		AllDependencies.Add(CurrentIndex, &bAlreadyExists);
+		if(!bAlreadyExists)
+		{
+			DependencyStack.Append(OwningTable.GetPluginInfoByIndex(CurrentIndex).GetDependencies());
+		}
 	}
 
 	// Find the set of dependencies which are directly referenced by something outside this set of plugins
-	TSet<int32> SharedDependencies;
+	TArray<int32> SharedDependencyStack;
 	for (int32 DependencyIndex : AllDependencies)
 	{
 		const FAssetTablePluginInfo& CurrentPluginInfo = OwningTable.GetPluginInfoByIndex(DependencyIndex);
@@ -1949,17 +1958,20 @@ void FAssetTablePluginInfo::ComputeDependencySizes(const FAssetTable& OwningTabl
 		}
 		if (IsShared)
 		{
-			SharedDependencies.Add(DependencyIndex);
+			SharedDependencyStack.Add(DependencyIndex);
 		}
 	}
-
+	
 	// Add to the SharedDependencies list all transitive dependencies
-	TArray<int32> SharedDependencyStack = SharedDependencies.Array();
+	TSet<int32> SharedDependencies;
 	while (SharedDependencyStack.Num() > 0)
 	{
 		int32 CurrentIndex = SharedDependencyStack.Pop();
-		SharedDependencies.Add(CurrentIndex);
-		SharedDependencyStack.Append(OwningTable.GetPluginInfoByIndex(CurrentIndex).GetDependencies());
+		SharedDependencies.Add(CurrentIndex, &bAlreadyExists);
+		if (!bAlreadyExists)
+		{
+			SharedDependencyStack.Append(OwningTable.GetPluginInfoByIndex(CurrentIndex).GetDependencies());
+		}
 	}
 
 	// UniqueDependencies are all the dependencies that aren't shared
diff --git a/Engine/Plugins/Editor/AssetManagerEditor/Source/AssetManagerEditor/Private/TreeView/AssetTable.h b/Engine/Plugins/Editor/AssetManagerEditor/Source/AssetManagerEditor/Private/TreeView/AssetTable.h
--- a/Engine/Plugins/Editor/AssetManagerEditor/Source/AssetManagerEditor/Private/TreeView/AssetTable.h
+++ b/Engine/Plugins/Editor/AssetManagerEditor/Source/AssetManagerEditor/Private/TreeView/AssetTable.h
@@ -161,6 +161,10 @@ public:
 
 	static void ComputeTotalSelfAndInclusiveSizes(const FAssetTable& OwningTable, const TSet<int32>& RootPlugins, int64& OutTotalSelfSize, int64& OutTotalInclusiveSize);
 
+	TSet<int32> GetPluginAssets() const { return PluginAssets; };
+
+	void AddPluginAsset(const int32 RowIndex) { PluginAssets.Add(RowIndex); };
+
 	template<typename T> const T* TryGetDataByKey(int32 Key) const
 	{
 		// In general there are expected to be very few keys since each key corresponds to a custom implemented
@@ -185,6 +189,7 @@ private:
 
 	TArray<int32> PluginDependencies;
 	TArray<int32> PluginReferencers;
+	TSet<int32> PluginAssets;
 
 	typedef TVariant<bool, const TCHAR*> CustomFieldVariantType;
 
@@ -335,6 +340,8 @@ public:
 		CustomColumns.Empty();
 	}
 
+	void PopulateAssetsInPlugins();
+
 	enum class ECustomColumnDefinitionType
 	{
 		Boolean,
diff --git a/Engine/Plugins/Editor/AssetManagerEditor/Source/AssetManagerEditor/Private/TreeView/SAssetTableTreeView.cpp b/Engine/Plugins/Editor/AssetManagerEditor/Source/AssetManagerEditor/Private/TreeView/SAssetTableTreeView.cpp
--- a/Engine/Plugins/Editor/AssetManagerEditor/Source/AssetManagerEditor/Private/TreeView/SAssetTableTreeView.cpp
+++ b/Engine/Plugins/Editor/AssetManagerEditor/Source/AssetManagerEditor/Private/TreeView/SAssetTableTreeView.cpp
@@ -2082,6 +2082,8 @@ void SAssetTableTreeView::RefreshAssets()
 				UE_LOG(LogInsights, Display, TEXT("Found root plugin %s"), StoredPluginName);
 				PluginInfo.bIsRootPlugin = true;
 			}
+
+			AssetTable->PopulateAssetsInPlugins();
 		}
 	}
 	else
