Added support for showing Verse paths to the reference viewer.

diff --git a/Engine/Plugins/Editor/AssetManagerEditor/Source/AssetManagerEditor/Private/ReferenceViewer/EdGraphNode_Reference.cpp b/Engine/Plugins/Editor/AssetManagerEditor/Source/AssetManagerEditor/Private/ReferenceViewer/EdGraphNode_Reference.cpp
--- a/Engine/Plugins/Editor/AssetManagerEditor/Source/AssetManagerEditor/Private/ReferenceViewer/EdGraphNode_Reference.cpp
+++ b/Engine/Plugins/Editor/AssetManagerEditor/Source/AssetManagerEditor/Private/ReferenceViewer/EdGraphNode_Reference.cpp
@@ -4,6 +4,7 @@
 
 #include "Algo/Sort.h"
 #include "Algo/Unique.h"
+#include "Containers/VersePath.h"
 #include "GenericPlatform/GenericPlatformFile.h"
 #include "EdGraph/EdGraphPin.h"
 #include "HAL/PlatformFileManager.h"
@@ -108,11 +109,6 @@ void UEdGraphNode_Reference::SetupReferenceNode(const FIntPoint& NodeLoc, const
 
 		// append the type so it shows up on the extra line
 		NodeTitle = FText::FromString(FString::Printf(TEXT("%s\n%s"), *MainAssetName, *AssetTypeName));
-
-		if (bIsPackage)
-		{
-			NodeComment = First.PackageName.ToString();
-		}
 	}
 	else
 	{
@@ -120,6 +116,7 @@ void UEdGraphNode_Reference::SetupReferenceNode(const FIntPoint& NodeLoc, const
 	}
 	
 	CacheAssetData(InAssetData);
+	UpdatePath();
 	AllocateDefaultPins();
 }
 
@@ -137,6 +134,7 @@ void UEdGraphNode_Reference::SetReferenceNodeCollapsed(const FIntPoint& NodeLoc,
 	NodeTitle = FText::Format( LOCTEXT("ReferenceNodeCollapsedTitle", "{0}"), FText::AsNumber(InNumReferencesExceedingMax));
 
 	CacheAssetData(FAssetData());
+	UpdatePath();
 	AllocateDefaultPins();
 }
 
@@ -210,43 +208,104 @@ FLinearColor UEdGraphNode_Reference::GetNodeTitleColor() const
 
 FText UEdGraphNode_Reference::GetTooltipText() const
 {
+	if (Identifiers.IsEmpty())
+	{
+		return FText::GetEmpty();
+	}
+
 	// Showing up to 15 nodes paths, in order to avoid an extremely long tooltip and huge widget
 	constexpr int32 MaxReferenceNum = 15;
 
 	FString TooltipString;
-	int32 ReferenceCount = 0;
 
 	if (IsCollapsed())
 	{
 		TooltipString.Append(GetNodeTitle(ENodeTitleType::FullTitle).ToString() + " collapsed nodes:\n");
 	}
 
-	for (const FAssetIdentifier& AssetId : Identifiers)
+	const bool bIsShowingContentVersePath = GetReferenceViewerGraph()->IsShowingContentVersePath();
+	const bool bShowCollections = !IsCollapsed() && bIsPackage;
+
+	auto IsPackage = [](const FAssetIdentifier& AssetId)
+	{
+		return !AssetId.GetPrimaryAssetId().IsValid() && !AssetId.IsValue();
+	};
+
+	// Get the asset data we need once.
+	TArray<FName> PackageNames;
+	if (bIsShowingContentVersePath)
 	{
-		if (!TooltipString.IsEmpty())
+		PackageNames.Reserve(FMath::Min(Identifiers.Num(), MaxReferenceNum));
+
+		int32 ReferenceCount = 0;
+		for (const FAssetIdentifier& AssetId : Identifiers)
 		{
-			TooltipString.Append(TEXT("\n"));
+			if (IsPackage(AssetId))
+			{
+				PackageNames.Add(AssetId.PackageName);
+			}
+
+			if (++ReferenceCount > MaxReferenceNum)
+			{
+				// We only need the asset data for the first MaxReferenceNum identifiers since that is all we show.
+				break;
+			}
 		}
-		TooltipString.Append(AssetId.ToString());
 
-		ReferenceCount++;
+		PackageNames.Sort(FNameFastLess());
+		PackageNames.SetNum(Algo::Unique(PackageNames));
+	}
+	else if (bShowCollections)
+	{
+		PackageNames.Add(Identifiers[0].PackageName);
+	}
+
+	TMap<FName, FAssetData> Assets;
+	UE::AssetRegistry::GetAssetForPackages(PackageNames, Assets);
 
-		// Avoiding an extremely long tooltip
-		if (ReferenceCount > MaxReferenceNum)
+	{
+		int32 ReferenceCount = 0;
+		for (const FAssetIdentifier& AssetId : Identifiers)
 		{
-			TooltipString.Append(TEXT("\n..."));
-			break;
+			if (!TooltipString.IsEmpty())
+			{
+				TooltipString.Append(TEXT("\n"));
+			}
+
+			UE::Core::FVersePath VersePath;
+
+			if (bIsShowingContentVersePath && IsPackage(AssetId))
+			{
+				if (const FAssetData* AssetData = Assets.Find(AssetId.PackageName))
+				{
+					VersePath = AssetData->GetVersePath();
+				}
+			}
+
+			if (VersePath.IsValid())
+			{
+				TooltipString.Append(VersePath.ToString());
+			}
+			else
+			{
+				TooltipString.Append(AssetId.ToString());
+			}
+
+			// Avoiding an extremely long tooltip
+			if (++ReferenceCount > MaxReferenceNum)
+			{
+				TooltipString.Append(TEXT("\n..."));
+				break;
+			}
 		}
 	}
 
-	if (!IsCollapsed() && !Identifiers.IsEmpty())
+	if (bShowCollections)
 	{
 		const FAssetIdentifier& AssetId = Identifiers[0];
 
 		// Retrieve Collections information
-		TMap<FName, FAssetData> Assets;
-		UE::AssetRegistry::GetAssetForPackages({AssetId.PackageName}, Assets);
-		if (FAssetData* AssetData = Assets.Find(AssetId.PackageName))
+		if (const FAssetData* AssetData = Assets.Find(AssetId.PackageName))
 		{
 			ICollectionManager& CollectionManager = FCollectionManagerModule::GetModule().Get();
 
@@ -288,7 +347,7 @@ FText UEdGraphNode_Reference::GetTooltipText() const
 		}
 	}
 
-	return FText::FromString(TooltipString);
+	return FText::FromString(MoveTemp(TooltipString));
 }
 
 FSlateIcon UEdGraphNode_Reference::GetIconAndTint(FLinearColor& OutColor) const
@@ -329,6 +388,32 @@ UEdGraphPin* UEdGraphNode_Reference::GetReferencerPin()
 	return ReferencerPin;
 }
 
+void UEdGraphNode_Reference::UpdatePath()
+{
+	if (!IsCollapsed() && Identifiers.Num() == 1)
+	{
+		UE::Core::FVersePath VersePath;
+
+		if (bIsPackage && CachedAssetData.IsValid() && GetReferenceViewerGraph()->IsShowingContentVersePath())
+		{
+			VersePath = CachedAssetData.GetVersePath();
+		}
+
+		if (VersePath.IsValid())
+		{
+			NodeComment = MoveTemp(VersePath).ToString();
+		}
+		else
+		{
+			NodeComment = Identifiers[0].PackageName.ToString();
+		}
+	}
+	else
+	{
+		NodeComment.Reset();
+	}
+}
+
 void UEdGraphNode_Reference::CacheAssetData(const FAssetData& AssetData)
 {
 	if ( AssetData.IsValid() && IsPackage() )
@@ -371,7 +456,7 @@ void UEdGraphNode_Reference::CacheAssetData(const FAssetData& AssetData)
 
 }
 
-FAssetData UEdGraphNode_Reference::GetAssetData() const
+const FAssetData& UEdGraphNode_Reference::GetAssetData() const
 {
 	return CachedAssetData;
 }
diff --git a/Engine/Plugins/Editor/AssetManagerEditor/Source/AssetManagerEditor/Private/ReferenceViewer/EdGraph_ReferenceViewer.cpp b/Engine/Plugins/Editor/AssetManagerEditor/Source/AssetManagerEditor/Private/ReferenceViewer/EdGraph_ReferenceViewer.cpp
--- a/Engine/Plugins/Editor/AssetManagerEditor/Source/AssetManagerEditor/Private/ReferenceViewer/EdGraph_ReferenceViewer.cpp
+++ b/Engine/Plugins/Editor/AssetManagerEditor/Source/AssetManagerEditor/Private/ReferenceViewer/EdGraph_ReferenceViewer.cpp
@@ -56,6 +56,7 @@ int32 FReferenceNodeInfo::ProvisionSize(const FAssetIdentifier& InParentId) cons
 
 UEdGraph_ReferenceViewer::UEdGraph_ReferenceViewer(const FObjectInitializer& ObjectInitializer)
 	: Super(ObjectInitializer)
+	, bShowingContentVersePath(false)
 {
 	if (!IsTemplate())
 	{
@@ -151,6 +152,42 @@ UEdGraphNode_Reference* UEdGraph_ReferenceViewer::RebuildGraph()
 	return NewRootNode;
 }
 
+void UEdGraph_ReferenceViewer::SetShowingContentVersePath(bool bInShowingContentVersePath)
+{
+	if (bShowingContentVersePath != bInShowingContentVersePath)
+	{
+		bShowingContentVersePath = bInShowingContentVersePath;
+
+		UpdatePaths();
+	}
+}
+
+void UEdGraph_ReferenceViewer::UpdatePaths()
+{
+	for (const TObjectPtr<UEdGraphNode>& Node : Nodes)
+	{
+		if (UEdGraphNode_Reference* ReferenceNode = Cast<UEdGraphNode_Reference>(Node))
+		{
+			ReferenceNode->UpdatePath();
+		}
+	}
+}
+
+void UEdGraph_ReferenceViewer::SetIsAssetIdentifierPassingSearchFilterCallback(const TOptional<FIsAssetIdentifierPassingSearchFilterCallback>& InIsAssetIdentifierPassingSearchFilterCallback)
+{
+	if (InIsAssetIdentifierPassingSearchFilterCallback.IsSet())
+	{
+		DoesAssetPassSearchFilterCallback = [IsAssetIdentifierPassingSearchFilterCallback = InIsAssetIdentifierPassingSearchFilterCallback.GetValue()](const FAssetIdentifier& InAssetIdentifier, const FAssetData& InAssetData)
+		{
+			return IsAssetIdentifierPassingSearchFilterCallback(InAssetIdentifier);
+		};
+	}
+	else
+	{
+		DoesAssetPassSearchFilterCallback.Reset();
+	}
+}
+
 FName UEdGraph_ReferenceViewer::GetCurrentCollectionFilter() const
 {
 	return CurrentCollectionFilterName;
@@ -290,32 +327,6 @@ UEdGraphNode_Reference* UEdGraph_ReferenceViewer::ConstructNodes(const TArray<FA
 			RecursivelyPopulateNodeInfos(false, GraphRootIdentifiers, NewDependencyNodeInfos, 0, Settings->GetSearchDependencyDepthLimit());
 		}
 
-		TSet<FName> AllPackageNames;
-		auto AddPackage = [](const FAssetIdentifier& AssetId, TSet<FName>& PackageNames)
-		{ 
-			// Only look for asset data if this is a package
-			if (!AssetId.IsValue() && !AssetId.PackageName.IsNone())
-			{
-				PackageNames.Add(AssetId.PackageName);
-			}
-		};
-
-		if (Settings->IsShowReferencers())
-		{
-			for (TPair<FAssetIdentifier, FReferenceNodeInfo>&  InfoPair : NewReferenceNodeInfos)
-			{
-				AddPackage(InfoPair.Key, AllPackageNames);
-			}
-		}
-
-		if (Settings->IsShowDependencies())
-		{
-			for (TPair<FAssetIdentifier, FReferenceNodeInfo>&  InfoPair : NewDependencyNodeInfos)
-			{
-				AddPackage(InfoPair.Key, AllPackageNames);
-			}
-		}
-
 		// Store the AssetData in the NodeInfos if needed, and collect Asset Type UClasses to populate the filters
 		TSet<FTopLevelAssetPath> AllClasses;
 		for (TPair<FAssetIdentifier, FReferenceNodeInfo>&  InfoPair : NewReferenceNodeInfos)
@@ -913,7 +924,7 @@ void UEdGraph_ReferenceViewer::RecursivelyFilterNodeInfos(const FAssetIdentifier
 	}
 
 	bool PassedAssetTypeFilter = FilterCollection && Settings->GetFiltersEnabled() ? FilterCollection->PassesAllFilters(NodeInfo) : true;
-	bool PassedSearchTextFilter = IsAssetPassingSearchTextFilter(InAssetId);
+	bool PassedSearchTextFilter = DoesAssetPassSearchTextFilter(InAssetId, NodeInfo.AssetData);
 
 	// Don't apply filters in Find Path Mode. Otherwise, check the type and search filters, and also don't include any assets in the central selection (where InCurrentDepth == 0)
 	bool PassedAllFilters = Settings->GetFindPathEnabled() || (PassedAssetTypeFilter && PassedSearchTextFilter && (CurrentDepth == 0 || !CurrentGraphRootIdentifiers.Contains(InAssetId)));
@@ -1121,9 +1132,9 @@ bool UEdGraph_ReferenceViewer::IsPackageIdentifierPassingPluginFilter(const FAss
 	return false;
 }
 
-bool UEdGraph_ReferenceViewer::IsAssetPassingSearchTextFilter(const FAssetIdentifier& InAssetIdentifier) const
+bool UEdGraph_ReferenceViewer::DoesAssetPassSearchTextFilter(const FAssetIdentifier& InAssetIdentifier, const FAssetData& InAssetData) const
 {
-	if (Settings->IsShowFilteredPackagesOnly() && IsAssetIdentifierPassingSearchFilterCallback.IsSet() && !(*IsAssetIdentifierPassingSearchFilterCallback)(InAssetIdentifier))
+	if (Settings->IsShowFilteredPackagesOnly() && DoesAssetPassSearchFilterCallback.IsSet() && !DoesAssetPassSearchFilterCallback(InAssetIdentifier, InAssetData))
 	{
 		return false;
 	}
diff --git a/Engine/Plugins/Editor/AssetManagerEditor/Source/AssetManagerEditor/Private/ReferenceViewer/HistoryManager.cpp b/Engine/Plugins/Editor/AssetManagerEditor/Source/AssetManagerEditor/Private/ReferenceViewer/HistoryManager.cpp
--- a/Engine/Plugins/Editor/AssetManagerEditor/Source/AssetManagerEditor/Private/ReferenceViewer/HistoryManager.cpp
+++ b/Engine/Plugins/Editor/AssetManagerEditor/Source/AssetManagerEditor/Private/ReferenceViewer/HistoryManager.cpp
@@ -2,7 +2,6 @@
 
 #include "HistoryManager.h"
 #include "AssetRegistry/AssetIdentifier.h"
-#include "Framework/MultiBox/MultiBoxBuilder.h"
 
 #define LOCTEXT_NAMESPACE "ReferenceViewer"
 
@@ -115,64 +114,22 @@ bool FReferenceViewerHistoryManager::CanGoBack() const
 	return HistoryData.Num() > 0 && CurrentHistoryIndex > 0;
 }
 
-FText FReferenceViewerHistoryManager::GetBackDesc() const
+const FReferenceViewerHistoryData* FReferenceViewerHistoryManager::GetBackHistoryData() const
 {
 	if ( CanGoBack() )
 	{
-		return HistoryData[CurrentHistoryIndex - 1].HistoryDesc;
+		return &HistoryData[CurrentHistoryIndex - 1];
 	}
-	return FText::GetEmpty();
+	return nullptr;
 }
 
-FText FReferenceViewerHistoryManager::GetForwardDesc() const
+const FReferenceViewerHistoryData* FReferenceViewerHistoryManager::GetForwardHistoryData() const
 {
 	if ( CanGoForward() )
 	{
-		return HistoryData[CurrentHistoryIndex + 1].HistoryDesc;
+		return &HistoryData[CurrentHistoryIndex + 1];
 	}
-	return FText::GetEmpty();
-}
-
-void FReferenceViewerHistoryManager::GetAvailableHistoryMenuItems(bool bGetPrior, FMenuBuilder& MenuBuilder)
-{
-	const FText HistoryHeadingString = (bGetPrior)? LOCTEXT("BackHistory", "Back History") : LOCTEXT("NextHistory", "Next History");
-	MenuBuilder.BeginSection("HistoryBackNext", HistoryHeadingString);
-	{
-		if (HistoryData.Num() > 1)
-		{
-			// if there is at least 2 history items...
-
-			// Start index is the first snapshot we should make a menu item out of
-			int32 StartIndex = 0;
-			// EndIndex is the last snapshot we should make a menu item out of
-			int32 EndIndex = CurrentHistoryIndex;
-
-			if (!bGetPrior)
-			{
-				// Need to return only items on or after the current history index
-				StartIndex = CurrentHistoryIndex;
-				EndIndex = HistoryData.Num() - 1;
-			}
-
-			// Check to make sure the start and end indices are within the bounds of the history list
-			if (StartIndex < HistoryData.Num() && EndIndex != -1)
-			{
-				// Get all menu items between and including the start index and end index
-				for (int32 HistoryIdx = StartIndex; HistoryIdx <= EndIndex; ++HistoryIdx)
-				{
-					MenuBuilder.AddMenuEntry(
-						HistoryData[HistoryIdx].HistoryDesc,
-						FText(),
-						FSlateIcon(),
-						FUIAction(
-							FExecuteAction::CreateRaw( this, &FReferenceViewerHistoryManager::ExecuteJumpToHistory, HistoryIdx )
-							)
-						);
-				}
-			}
-		}
-	}
-	MenuBuilder.EndSection();
+	return nullptr;
 }
 
 void FReferenceViewerHistoryManager::ApplyCurrentHistoryData()
diff --git a/Engine/Plugins/Editor/AssetManagerEditor/Source/AssetManagerEditor/Private/ReferenceViewer/HistoryManager.h b/Engine/Plugins/Editor/AssetManagerEditor/Source/AssetManagerEditor/Private/ReferenceViewer/HistoryManager.h
--- a/Engine/Plugins/Editor/AssetManagerEditor/Source/AssetManagerEditor/Private/ReferenceViewer/HistoryManager.h
+++ b/Engine/Plugins/Editor/AssetManagerEditor/Source/AssetManagerEditor/Private/ReferenceViewer/HistoryManager.h
@@ -11,9 +11,6 @@ class FMenuBuilder;
 /** The history data object, storing all important history data */
 struct FReferenceViewerHistoryData
 {
-	/** History Description */
-	FText HistoryDesc;
-
 	/** The list of package names to serve as the root */
 	TArray<FAssetIdentifier> Identifiers;
 };
@@ -56,19 +53,10 @@ public:
 	bool CanGoBack() const;
 
 	/** Gets the description of the previous history entry */
-	FText GetBackDesc() const;
+	const FReferenceViewerHistoryData* GetBackHistoryData() const;
 
 	/** Gets the description of the next history entry */
-	FText GetForwardDesc() const;
-
-	/**
-	 * Populates a list of menu items that can be added to a context menu
-	 * to allow a user to jump to different history snapshots instead of using the back and forward buttons
-	 *
-	 * @param bGetPrior If true gets history snapshots prior to the current history index(for navigating back). If false get history snapshots after the current history index (for navigating forward).
-	 * @param MenuBuilder The menubuilder to populate with menu items
-	 */
-	void GetAvailableHistoryMenuItems(bool bGetPrior, FMenuBuilder& MenuBuilder);
+	const FReferenceViewerHistoryData* GetForwardHistoryData() const;
 
 private:
 	/** Notifies the owner to update to the state described by the current history data */
diff --git a/Engine/Plugins/Editor/AssetManagerEditor/Source/AssetManagerEditor/Private/ReferenceViewer/SReferenceViewer.cpp b/Engine/Plugins/Editor/AssetManagerEditor/Source/AssetManagerEditor/Private/ReferenceViewer/SReferenceViewer.cpp
--- a/Engine/Plugins/Editor/AssetManagerEditor/Source/AssetManagerEditor/Private/ReferenceViewer/SReferenceViewer.cpp
+++ b/Engine/Plugins/Editor/AssetManagerEditor/Source/AssetManagerEditor/Private/ReferenceViewer/SReferenceViewer.cpp
@@ -10,12 +10,17 @@
 #include "ReferenceViewer/EdGraph_ReferenceViewer.h"
 #include "Widgets/Input/SSpinBox.h"
 
+#include "Algo/Sort.h"
+#include "Algo/Unique.h"
+#include "AssetToolsModule.h"
+#include "Containers/VersePath.h"
 #include "ReferenceViewer/HistoryManager.h"
 #include "ReferenceViewerStyle.h"
 #include "ReferenceViewer/EdGraphNode_Reference.h"
 #include "ReferenceViewer/EdGraphNode_ReferencedProperties.h"
 #include "ReferenceViewer/ReferenceViewerSchema.h"
 #include "AssetRegistry/AssetRegistryModule.h"
+#include "IAssetTools.h"
 #include "ICollectionContainer.h"
 #include "ICollectionManager.h"
 #include "ICollectionSource.h"
@@ -43,12 +48,17 @@
 
 #define LOCTEXT_NAMESPACE "ReferenceViewer"
 
-bool IsAssetIdentifierPassingSearchTextFilter(const FAssetIdentifier& InNode, const TArray<FString>& InSearchWords)
+bool DoesAssetPassSearchTextFilter(const FAssetIdentifier& InNode, const FAssetData& AssetData, bool bShowingContentVersePath, const TArray<FString>& InSearchWords)
 {
 	FString NodeString = InNode.ToString();
+	UE::Core::FVersePath VersePath;
+	if (bShowingContentVersePath)
+	{
+		VersePath = AssetData.GetVersePath();
+	}
 	for (const FString& Word : InSearchWords)
 	{
-		if (!NodeString.Contains(Word))
+		if (!NodeString.Contains(Word) && !VersePath.ToString().Contains(Word))
 		{
 			return false;
 		}
@@ -57,7 +67,7 @@ bool IsAssetIdentifierPassingSearchTextFilter(const FAssetIdentifier& InNode, co
 	return true;
 }
 
-EAppReturnType::Type ShowAssetsNeedsToLoadMessage(const TSet<FAssetData>& UnloadedAssetsData)
+EAppReturnType::Type ShowAssetsNeedsToLoadMessage(const TSet<FAssetData>& UnloadedAssetsData, bool bShowingContentVersePath)
 {
 	FString UnloadedAssetsNames;
 
@@ -71,7 +81,21 @@ EAppReturnType::Type ShowAssetsNeedsToLoadMessage(const TSet<FAssetData>& Unload
 			break;
 		}
 
-		UnloadedAssetsNames += TEXT("\n\n") + Data.GetFullName();
+		UnloadedAssetsNames += TEXT("\n\n");
+
+		if (bShowingContentVersePath)
+		{
+			const UE::Core::FVersePath VersePath = Data.GetVersePath();
+			if (VersePath.IsValid())
+			{
+				UnloadedAssetsNames += Data.AssetClassPath.ToString();
+				UnloadedAssetsNames += TEXT(' ');
+				UnloadedAssetsNames += VersePath.ToString();
+				continue;
+			}
+		}
+
+		UnloadedAssetsNames += Data.GetFullName();
 	}
 
 	if (UnloadedAssetsData.Num() > MaxAssetsShown)
@@ -115,6 +139,7 @@ SReferenceViewer::~SReferenceViewer()
 
 void SReferenceViewer::Construct(const FArguments& InArgs)
 {
+	bShowingContentVersePath = FAssetToolsModule::GetModule().Get().ShowingContentVersePath();
 	bRebuildingFilters = false;
 	bNeedsGraphRebuild = false;
 	bNeedsGraphRefilter = false;
@@ -133,6 +158,7 @@ void SReferenceViewer::Construct(const FArguments& InArgs)
 	GraphObj->Schema = UReferenceViewerSchema::StaticClass();
 	GraphObj->AddToRoot();
 	GraphObj->SetReferenceViewer(StaticCastSharedRef<SReferenceViewer>(AsShared()));
+	GraphObj->SetShowingContentVersePath(bShowingContentVersePath);
 	GraphObj->OnAssetsChanged().AddSP(this, &SReferenceViewer::OnUpdateFilterBar);
 
 	SGraphEditor::FGraphEditorEvents GraphEvents;
@@ -268,7 +294,7 @@ void SReferenceViewer::Construct(const FArguments& InArgs)
 					.ButtonContent()
 					[
 						SNew(STextBlock)
-						.Text_Lambda([this] { return FindPathAssetId.IsValid() ? FText::FromString(FindPathAssetId.ToString()) : LOCTEXT("ChooseTargetAsset", "Choose a target asset ... "); } )
+						.Text_Lambda([this] { return FindPathAssetId.IsValid() ? GetIdentifierText(FindPathAssetId) : LOCTEXT("ChooseTargetAsset", "Choose a target asset ... "); } )
 					]
 				]
 			]
@@ -653,6 +679,21 @@ void SReferenceViewer::Construct(const FArguments& InArgs)
 
 void SReferenceViewer::Tick( const FGeometry& AllottedGeometry, const double InCurrentTime, const float InDeltaTime )
 {
+	{
+		bool bNewShowingContentVersePath = FAssetToolsModule::GetModule().Get().ShowingContentVersePath();
+		if (bShowingContentVersePath != bNewShowingContentVersePath)
+		{
+			bShowingContentVersePath = bNewShowingContentVersePath;
+
+			if (GraphObj)
+			{
+				GraphObj->SetShowingContentVersePath(bShowingContentVersePath);
+			}
+
+			UpdateIsPassingSearchFilterCallback();
+		}
+	}
+
 	if (bNeedsGraphRebuild && (InCurrentTime - SliderDelayLastMovedTime > GraphRebuildSliderDelay))
 	{
 		bNeedsGraphRebuild = false;
@@ -836,6 +877,12 @@ void SReferenceViewer::RebuildGraph()
 			AssetRegistryModule.Get().OnAssetAdded().AddSP(this, &SReferenceViewer::OnAssetRegistryChanged);
 			AssetRegistryModule.Get().OnAssetRemoved().AddSP(this, &SReferenceViewer::OnAssetRegistryChanged);
 		}
+
+		IPluginManager& PluginManager = IPluginManager::Get();
+		if (!PluginManager.OnPluginEdited().IsBoundToObject(this))
+		{
+			PluginManager.OnPluginEdited().AddSP(this, &SReferenceViewer::OnPluginEdited);
+		}
 	}
 }
 
@@ -896,18 +943,20 @@ void SReferenceViewer::GraphNavigateHistoryForward()
 
 FText SReferenceViewer::GetHistoryBackTooltip() const
 {
-	if ( HistoryManager.CanGoBack() )
+	const FReferenceViewerHistoryData* HistoryData = HistoryManager.GetBackHistoryData();
+	if (HistoryData != nullptr)
 	{
-		return FText::Format( LOCTEXT("HistoryBackTooltip", "Back to {0}"), HistoryManager.GetBackDesc() );
+		return FText::Format( LOCTEXT("HistoryBackTooltip", "Back to {0}"), GetIdentifierSummaryText(HistoryData->Identifiers) );
 	}
 	return FText::GetEmpty();
 }
 
 FText SReferenceViewer::GetHistoryForwardTooltip() const
 {
-	if ( HistoryManager.CanGoForward() )
+	const FReferenceViewerHistoryData* HistoryData = HistoryManager.GetForwardHistoryData();
+	if (HistoryData != nullptr)
 	{
-		return FText::Format( LOCTEXT("HistoryForwardTooltip", "Forward to {0}"), HistoryManager.GetForwardDesc() );
+		return FText::Format( LOCTEXT("HistoryForwardTooltip", "Forward to {0}"), GetIdentifierSummaryText(HistoryData->Identifiers) );
 	}
 	return FText::GetEmpty();
 }
@@ -918,23 +967,7 @@ FText SReferenceViewer::GetAddressBarText() const
 	{
 		if (TemporaryPathBeingEdited.IsEmpty())
 		{
-			const TArray<FAssetIdentifier>& CurrentGraphRootPackageNames = GraphObj->GetCurrentGraphRootIdentifiers();
-			if (CurrentGraphRootPackageNames.Num() == 1)
-			{
-				return FText::FromString(CurrentGraphRootPackageNames[0].ToString());
-			}
-			else if (CurrentGraphRootPackageNames.Num() > 1)
-			{
-				return FText::Format(
-					LOCTEXT("AddressBarMultiplePackagesText", "{0} and {1} others"),
-					FText::FromString(CurrentGraphRootPackageNames[0].ToString()),
-					FText::AsNumber(CurrentGraphRootPackageNames.Num() - 1)
-				);
-			}
-			else
-			{
-				return LOCTEXT("NoAssetFoundText", "No Assets Found");
-			}
+			return GetIdentifierSummaryText(GraphObj->GetCurrentGraphRootIdentifiers());
 		}
 		else
 		{
@@ -945,6 +978,44 @@ FText SReferenceViewer::GetAddressBarText() const
 	return FText();
 }
 
+FText SReferenceViewer::GetIdentifierSummaryText(const TArray<FAssetIdentifier>& Identifiers) const
+{
+	if (Identifiers.Num() == 1)
+	{
+		return GetIdentifierText(Identifiers[0]);
+	}
+	else if (Identifiers.Num() > 1)
+	{
+		return FText::Format(
+			LOCTEXT("AddressBarMultiplePackagesText", "{0} and {1} others"),
+			GetIdentifierText(Identifiers[0]),
+			FText::AsNumber(Identifiers.Num() - 1)
+		);
+	}
+	else
+	{
+		return LOCTEXT("NoAssetFoundText", "No Assets Found");
+	}
+}
+
+FText SReferenceViewer::GetIdentifierText(const FAssetIdentifier& Identifier) const
+{
+	if (bShowingContentVersePath && !Identifier.GetPrimaryAssetId().IsValid() && !Identifier.IsValue())
+	{
+		TMap<FName, FAssetData> Assets;
+		UE::AssetRegistry::GetAssetForPackages({ Identifier.PackageName }, Assets);
+		if (const FAssetData* Asset = Assets.Find(Identifier.PackageName))
+		{
+			UE::Core::FVersePath VersePath = Asset->GetVersePath();
+			if (VersePath.IsValid())
+			{
+				return FText::FromString(MoveTemp(VersePath).ToString());
+			}
+		}
+	}
+	return FText::FromString(Identifier.ToString());
+}
+
 FText SReferenceViewer::GetStatusText() const
 {
 	FString DirtyPackages;
@@ -1021,14 +1092,35 @@ void SReferenceViewer::OnAddressBarTextCommitted(const FText& NewText, ETextComm
 
 		// Using GetDependencies just to check if NewPath exists in the dependency tree. We do not actually care about the dependencies here
 		TArray<FAssetIdentifier> UnusedDependencies;
-		if (AssetRegistry->GetDependencies(NewPath, UnusedDependencies))
+		if (NewPath.IsValid() && AssetRegistry->GetDependencies(NewPath, UnusedDependencies))
 		{
-			NewPaths.Add(FAssetIdentifier::FromString(NewText.ToString()));
+			NewPaths.Add(NewPath);
+		}
+		else if (bShowingContentVersePath)
+		{
+			// Syntactically the grammer for FAssetIdentifiers and FVersePaths overlap.
+			// If the we couldn't find NewPath it may be because it is a Verse path, so check that next.
+			UE::Core::FVersePath VersePath;
+			if (UE::Core::FVersePath::TryMake(VersePath, NewText.ToString()))
+			{
+				FAssetData AssetData = FAssetToolsModule::GetModule().Get().FindAssetByVersePath(VersePath);
+				if (AssetData.IsValid())
+				{
+					NewPath = FAssetIdentifier(AssetData.PackageName);
+					if (AssetRegistry->GetDependencies(NewPath, UnusedDependencies))
+					{
+						NewPaths.Add(NewPath);
+					}
+				}
+			}
 		}
 
-		SetGraphRootIdentifiers(NewPaths);
+		if (CurrentlyVisualizedAssets != NewPaths)
+		{
+			SetGraphRootIdentifiers(NewPaths);
 
-		OnReferenceViewerSelectionChanged().Broadcast(CurrentlyVisualizedAssets, GraphObj->GetCurrentGraphRootIdentifiers());
+			OnReferenceViewerSelectionChanged().Broadcast(CurrentlyVisualizedAssets, GraphObj->GetCurrentGraphRootIdentifiers());
+		}
 	}
 }
 
@@ -1058,12 +1150,10 @@ void SReferenceViewer::OnUpdateHistoryData(FReferenceViewerHistoryData& HistoryD
 	if ( GraphObj )
 	{
 		const TArray<FAssetIdentifier>& CurrentGraphRootIdentifiers = GraphObj->GetCurrentGraphRootIdentifiers();
-		HistoryData.HistoryDesc = GetAddressBarText();
 		HistoryData.Identifiers = CurrentGraphRootIdentifiers;
 	}
 	else
 	{
-		HistoryData.HistoryDesc = FText::GetEmpty();
 		HistoryData.Identifiers.Empty();
 	}
 }
@@ -1303,7 +1393,7 @@ ECheckBoxState SReferenceViewer::IsEnablePluginFilterChecked() const
 	return Settings->GetEnablePluginFilter() ? ECheckBoxState::Checked : ECheckBoxState::Unchecked;
 }
 
-void SReferenceViewer::PluginFilterAddMenuEntry(FMenuBuilder& MenuBuilder, const FName& PluginName)
+void SReferenceViewer::PluginFilterAddMenuEntry(FMenuBuilder& MenuBuilder, const FName& PluginName, const FText& Label, const FText& ToolTip)
 {
 	FExecuteAction ActionClicked = FExecuteAction::CreateLambda([this, PluginName]()
 	{
@@ -1336,8 +1426,8 @@ void SReferenceViewer::PluginFilterAddMenuEntry(FMenuBuilder& MenuBuilder, const
 	});
 
 	MenuBuilder.AddMenuEntry(
-		FText::FromName(PluginName),
-		FText::FromName(PluginName),
+		Label,
+		ToolTip,
 		FSlateIcon(),
 		FUIAction(ActionClicked, FCanExecuteAction(), ActionChecked),
 		NAME_None, //InExtensionHook
@@ -1377,53 +1467,78 @@ TSharedRef<SWidget> SReferenceViewer::BuildPluginFilterMenu()
 		})
 	);
 
-	const TArray<FName> InitialPluginFilter = GraphObj->GetCurrentPluginFilter();
-
-	TArray<FName> EnabledPluginNames;
-	for (const TSharedRef<IPlugin>& Plugin : IPluginManager::Get().GetEnabledPluginsWithContent())
+	TArray<FName> PluginNames = GraphObj->GetEncounteredPluginsAmongNodes();
+	if (PluginNames.IsEmpty())
 	{
-		const FName Name = FName(Plugin.Get().GetName());
-		EnabledPluginNames.AddUnique(Name);
+		return MenuBuilder.MakeWidget();
 	}
 
-	TArray<FName> PluginNames = GraphObj->GetEncounteredPluginsAmongNodes();
-	PluginNames.Sort([](const FName& A, const FName& B) { return A.Compare(B) < 0; });
-
-	if (PluginNames.Num() > 0)
+	// Create a map of plugin names to enabled plugins.
+	TMap<FName, TSharedRef<IPlugin>> EnabledPlugins;
 	{
-		MenuBuilder.AddSeparator();
+		TArray<TSharedRef<IPlugin>> EnabledPluginsWithContent = IPluginManager::Get().GetEnabledPluginsWithContent();
+		EnabledPlugins.Reserve(EnabledPluginsWithContent.Num());
+		for (TSharedRef<IPlugin>& Plugin : EnabledPluginsWithContent)
+		{
+			const FName Name = FName(Plugin->GetName());
+			EnabledPlugins.Add(Name, MoveTemp(Plugin));
+		}
 	}
 
-	int NumNonPluginsAdded = 0;
-	// First add any "plugins" that are not actual plugins, such as /Game and /Engine.
-	for (const FName& PluginName : PluginNames)
+	struct FPluginFilter
 	{
-		if (EnabledPluginNames.Contains(PluginName))
+		FPluginFilter(FName InPluginName, const TSharedRef<IPlugin>* InPlugin)
+			: PluginName(InPluginName)
+			, bIsPlugin(InPlugin != nullptr)
+			, ToolTip(FText::FromName(InPluginName))
 		{
-			continue;
+			if (InPlugin != nullptr)
+			{
+				Label = FText::FromString((*InPlugin)->GetFriendlyName());
+			}
+			else
+			{
+				Label = ToolTip;
+			}
 		}
-		
-		PluginFilterAddMenuEntry(MenuBuilder, PluginName);
 
-		++NumNonPluginsAdded;
-	}
+		FName PluginName;
+		bool bIsPlugin;
+		FText Label;
+		FText ToolTip;
+	};
 
-	// Add a separator between non-real plugins (such as /Game and /Engine) and actual plugins.
-	// Skip this if no non-real plugins were added or if there are no more plugins to add.
-	if (NumNonPluginsAdded > 0 && NumNonPluginsAdded < PluginNames.Num())
+	TArray<FPluginFilter> PluginFilters;
+	PluginFilters.Reserve(PluginNames.Num());
+
+	for (const FName& PluginName : PluginNames)
 	{
-		MenuBuilder.AddSeparator();
+		PluginFilters.Emplace(PluginName, EnabledPlugins.Find(PluginName));
 	}
+	
+	// Sort non-real plugins (such as /Game and /Engine) first, then by display label.
+	PluginFilters.Sort([](const FPluginFilter& A, const FPluginFilter& B)
+	{
+		if (A.bIsPlugin != B.bIsPlugin)
+		{
+			return !A.bIsPlugin;
+		}
+		return A.Label.CompareTo(B.Label) < 0;
+	});
 
-	// Add all real plugins.
-	for (const FName& PluginName : PluginNames)
+	MenuBuilder.AddSeparator();
+
+	for (int32 Index = 0; Index < PluginFilters.Num(); ++Index)
 	{
-		if (!EnabledPluginNames.Contains(PluginName))
+		const FPluginFilter& PluginFilter = PluginFilters[Index];
+
+		// Add a separator between non-real plugins (such as /Game and /Engine) and actual plugins.
+		if (Index > 0 && PluginFilters[Index - 1].bIsPlugin != PluginFilter.bIsPlugin)
 		{
-			continue;
+			MenuBuilder.AddSeparator();
 		}
-		
-		PluginFilterAddMenuEntry(MenuBuilder, PluginName);
+
+		PluginFilterAddMenuEntry(MenuBuilder, PluginFilter.PluginName, PluginFilter.Label, PluginFilter.ToolTip);
 	}
 
 	return MenuBuilder.MakeWidget();
@@ -1439,6 +1554,13 @@ FText SReferenceViewer::GetPluginComboButtonText() const
 	}
 	else if (CurrentPluginFilter.Num() == 1)
 	{
+		FNameBuilder NameBuilder;
+		CurrentPluginFilter[0].AppendString(NameBuilder);
+		TSharedPtr<IPlugin> Plugin = IPluginManager::Get().FindPlugin(NameBuilder);
+		if (Plugin)
+		{
+			return FText::FromString(Plugin->GetFriendlyName());
+		}
 		return FText::FromName(CurrentPluginFilter[0]);
 	}
 	else
@@ -1485,21 +1607,21 @@ void SReferenceViewer::UpdateIsPassingSearchFilterCallback()
 {
 	if (GraphObj)
 	{
-		TOptional<UEdGraph_ReferenceViewer::FIsAssetIdentifierPassingSearchFilterCallback> IsAssetIdentifierPassingSearchFilterCallback;
-		FString SearchString = SearchBox->GetText().ToString();
-		TArray<FString> SearchWords;
-		SearchString.ParseIntoArrayWS(SearchWords);
+		UEdGraph_ReferenceViewer::FDoesAssetPassSearchFilterCallback DoesAssetPassSearchFilterCallback;
+		if (Settings->IsShowFilteredPackagesOnly())
 		{
-			if (Settings->IsShowFilteredPackagesOnly())
+			FString SearchString = SearchBox->GetText().ToString();
+			TArray<FString> SearchWords;
+			SearchString.ParseIntoArrayWS(SearchWords);
+			if (SearchWords.Num() > 0)
 			{
-				if (SearchWords.Num() > 0)
+				DoesAssetPassSearchFilterCallback = [bShowingContentVersePath = bShowingContentVersePath, SearchWords](const FAssetIdentifier& InAssetIdentifier, const FAssetData& InAssetData)
 				{
-					IsAssetIdentifierPassingSearchFilterCallback = [=](const FAssetIdentifier& InAssetIdentifier) { return IsAssetIdentifierPassingSearchTextFilter(InAssetIdentifier, SearchWords); };
-				}
+					return DoesAssetPassSearchTextFilter(InAssetIdentifier, InAssetData, bShowingContentVersePath, SearchWords);
+				};
 			}
-
-			GraphObj->SetIsAssetIdentifierPassingSearchFilterCallback(IsAssetIdentifierPassingSearchFilterCallback);
 		}
+		GraphObj->SetDoesAssetPassSearchFilterCallback(DoesAssetPassSearchFilterCallback);
 		GraphObj->RefilterGraph();
 	}
 }
@@ -1929,8 +2051,12 @@ void SReferenceViewer::RegisterActions()
 					{
 						if (ReferenceNode->GetAssetData().IsValid())
 						{
+							if (!Result.IsEmpty())
+							{
+								Result += LINE_TERMINATOR;
+							}
+
 							Result += ReferenceNode->GetAssetData().PackageName.ToString();
-							Result += TEXT("\n");
 						}
 					}
 				}
@@ -1990,9 +2116,9 @@ void SReferenceViewer::ReCenterGraph()
 	ReCenterGraphOnNodes( GraphEditorPtr->GetSelectedNodes() );
 }
 
-FString SReferenceViewer::GetReferencedObjectsList() const
+FString SReferenceViewer::GetObjectsList(EObjectsListType ObjectsListType) const
 {
-	FString ReferencedObjectsList;
+	FString ObjectsList;
 
 	TSet<FName> AllSelectedPackageNames;
 	GetPackageNamesFromSelectedNodes(AllSelectedPackageNames);
@@ -2001,104 +2127,140 @@ FString SReferenceViewer::GetReferencedObjectsList() const
 	{
 		FAssetRegistryModule& AssetRegistryModule = FModuleManager::LoadModuleChecked<FAssetRegistryModule>(TEXT("AssetRegistry"));
 
+		TArray<FName> HardPackageNames;
+		TArray<FName> SoftPackageNames;
+		TArray<FName> AllPackageNames;
+		TMap<FName, FAssetData> Assets;
+
 		for (const FName& SelectedPackageName : AllSelectedPackageNames)
 		{
-			TArray<FName> HardDependencies;
-			AssetRegistryModule.Get().GetDependencies(SelectedPackageName, HardDependencies, UE::AssetRegistry::EDependencyCategory::Package, UE::AssetRegistry::EDependencyQuery::Hard);
-			
-			TArray<FName> SoftDependencies;
-			AssetRegistryModule.Get().GetDependencies(SelectedPackageName, SoftDependencies, UE::AssetRegistry::EDependencyCategory::Package, UE::AssetRegistry::EDependencyQuery::Soft);
+			HardPackageNames.Reset();
+			SoftPackageNames.Reset();
+			AllPackageNames.Reset();
+			Assets.Reset();
 
-			ReferencedObjectsList += FString::Printf(TEXT("[%s - Dependencies]\n"), *SelectedPackageName.ToString());
-			if (HardDependencies.Num() > 0)
+			const TCHAR* ObjectsListName = nullptr;
+			switch (ObjectsListType)
 			{
-				ReferencedObjectsList += TEXT("  [HARD]\n");
-				for (const FName& HardDependency : HardDependencies)
-				{
-					const FString PackageString = HardDependency.ToString();
-					ReferencedObjectsList += FString::Printf(TEXT("    %s.%s\n"), *PackageString, *FPackageName::GetLongPackageAssetName(PackageString));
-				}
+			case SReferenceViewer::EObjectsListType::Referenced:
+				ObjectsListName = TEXT("Dependencies");
+
+				AssetRegistryModule.Get().GetDependencies(SelectedPackageName, HardPackageNames, UE::AssetRegistry::EDependencyCategory::Package, UE::AssetRegistry::EDependencyQuery::Hard);
+				AssetRegistryModule.Get().GetDependencies(SelectedPackageName, SoftPackageNames, UE::AssetRegistry::EDependencyCategory::Package, UE::AssetRegistry::EDependencyQuery::Soft);
+				break;
+			case SReferenceViewer::EObjectsListType::Referencing:
+				ObjectsListName = TEXT("Referencers");
+
+				AssetRegistryModule.Get().GetReferencers(SelectedPackageName, HardPackageNames, UE::AssetRegistry::EDependencyCategory::Package, UE::AssetRegistry::EDependencyQuery::Hard);
+				AssetRegistryModule.Get().GetReferencers(SelectedPackageName, SoftPackageNames, UE::AssetRegistry::EDependencyCategory::Package, UE::AssetRegistry::EDependencyQuery::Soft);
+				break;
+			default:
+				checkNoEntry();
+				return {};
 			}
-			if (SoftDependencies.Num() > 0)
+
+			AllPackageNames.Add(SelectedPackageName);
+			AllPackageNames.Append(HardPackageNames);
+			AllPackageNames.Append(SoftPackageNames);
+
+			AllPackageNames.Sort(FNameFastLess());
+			AllPackageNames.SetNum(Algo::Unique(AllPackageNames));
+			
+			UE::AssetRegistry::GetAssetForPackages(AllPackageNames, Assets);
+
+			auto AppendPackageNames = [this, &ObjectsList, &Assets](const TCHAR* Title, const TArray<FName>& PackageNames)
 			{
-				ReferencedObjectsList += TEXT("  [SOFT]\n");
-				for (const FName& SoftDependency : SoftDependencies)
+				if (PackageNames.IsEmpty())
 				{
-					const FString PackageString = SoftDependency.ToString();
-					ReferencedObjectsList += FString::Printf(TEXT("    %s.%s\n"), *PackageString, *FPackageName::GetLongPackageAssetName(PackageString));
+					return;
 				}
-			}
-		}
-	}
 
-	return ReferencedObjectsList;
-}
+				ObjectsList += LINE_TERMINATOR;
+				ObjectsList += FString::Printf(TEXT("  [%s]"), Title);
 
-FString SReferenceViewer::GetReferencingObjectsList() const
-{
-	FString ReferencingObjectsList;
+				for (const FName& PackageName : PackageNames)
+				{
+					ObjectsList += LINE_TERMINATOR;
+					ObjectsList += TEXT("    ");
+					if (const FAssetData* Asset = Assets.Find(PackageName))
+					{
+						if (bShowingContentVersePath)
+						{
+							UE::Core::FVersePath VersePath = Asset->GetVersePath();
+							if (VersePath.IsValid())
+							{
+								ObjectsList += VersePath.ToString();
+								continue;
+							}
+						}
 
-	TSet<FName> AllSelectedPackageNames;
-	GetPackageNamesFromSelectedNodes(AllSelectedPackageNames);
+						ObjectsList += Asset->GetObjectPathString();
+					}
+					else
+					{
+						const FString PackageString = PackageName.ToString();
+						ObjectsList += FString::Printf(TEXT("%s.%s"), *PackageString, *FPackageName::GetLongPackageAssetName(PackageString));
+					}
+				}
+			};
 
-	if (AllSelectedPackageNames.Num() > 0)
-	{
-		FAssetRegistryModule& AssetRegistryModule = FModuleManager::LoadModuleChecked<FAssetRegistryModule>(TEXT("AssetRegistry"));
+			if (!ObjectsList.IsEmpty())
+			{
+				ObjectsList += LINE_TERMINATOR;
+			}
 
-		for (const FName& SelectedPackageName : AllSelectedPackageNames)
-		{
-			TArray<FName> HardDependencies;
-			AssetRegistryModule.Get().GetReferencers(SelectedPackageName, HardDependencies, UE::AssetRegistry::EDependencyCategory::Package, UE::AssetRegistry::EDependencyQuery::Hard);
+			FString SelectedPackageString;
+			{
+				UE::Core::FVersePath VersePath;
 
-			TArray<FName> SoftDependencies;
-			AssetRegistryModule.Get().GetReferencers(SelectedPackageName, SoftDependencies, UE::AssetRegistry::EDependencyCategory::Package, UE::AssetRegistry::EDependencyQuery::Soft);
+				if (bShowingContentVersePath)
+				{
+					if (const FAssetData* Asset = Assets.Find(SelectedPackageName))
+					{
+						VersePath = Asset->GetVersePath();
+					}
+				}
 
-			ReferencingObjectsList += FString::Printf(TEXT("[%s - Referencers]\n"), *SelectedPackageName.ToString());
-			if (HardDependencies.Num() > 0)
-			{
-				ReferencingObjectsList += TEXT("  [HARD]\n");
-				for (const FName& HardDependency : HardDependencies)
+				if (VersePath.IsValid())
 				{
-					const FString PackageString = HardDependency.ToString();
-					ReferencingObjectsList += FString::Printf(TEXT("    %s.%s\n"), *PackageString, *FPackageName::GetLongPackageAssetName(PackageString));
+					SelectedPackageString = MoveTemp(VersePath).ToString();
 				}
-			}
-			if (SoftDependencies.Num() > 0)
-			{
-				ReferencingObjectsList += TEXT("  [SOFT]\n");
-				for (const FName& SoftDependency : SoftDependencies)
+				else
 				{
-					const FString PackageString = SoftDependency.ToString();
-					ReferencingObjectsList += FString::Printf(TEXT("    %s.%s\n"), *PackageString, *FPackageName::GetLongPackageAssetName(PackageString));
+					SelectedPackageString = SelectedPackageName.ToString();
 				}
 			}
+
+			ObjectsList += FString::Printf(TEXT("[%s - %s]"), *SelectedPackageString, ObjectsListName);
+			AppendPackageNames(TEXT("HARD"), HardPackageNames);
+			AppendPackageNames(TEXT("SOFT"), SoftPackageNames);
 		}
 	}
 
-	return ReferencingObjectsList;
+	return ObjectsList;
 }
 
 void SReferenceViewer::CopyReferencedObjects()
 {
-	const FString ReferencedObjectsList = GetReferencedObjectsList();
+	const FString ReferencedObjectsList = GetObjectsList(EObjectsListType::Referenced);
 	FPlatformApplicationMisc::ClipboardCopy(*ReferencedObjectsList);
 }
 
 void SReferenceViewer::CopyReferencingObjects()
 {
-	const FString ReferencingObjectsList = GetReferencingObjectsList();
+	const FString ReferencingObjectsList = GetObjectsList(EObjectsListType::Referencing);
 	FPlatformApplicationMisc::ClipboardCopy(*ReferencingObjectsList);
 }
 
 void SReferenceViewer::ShowReferencedObjects()
 {
-	const FString ReferencedObjectsList = GetReferencedObjectsList();
+	const FString ReferencedObjectsList = GetObjectsList(EObjectsListType::Referenced);
 	SGenericDialogWidget::OpenDialog(LOCTEXT("ReferencedObjectsDlgTitle", "Referenced Objects"), SNew(STextBlock).Text(FText::FromString(ReferencedObjectsList)));
 }
 
 void SReferenceViewer::ShowReferencingObjects()
 {	
-	const FString ReferencingObjectsList = GetReferencingObjectsList();
+	const FString ReferencingObjectsList = GetObjectsList(EObjectsListType::Referencing);
 	SGenericDialogWidget::OpenDialog(LOCTEXT("ReferencingObjectsDlgTitle", "Referencing Objects"), SNew(STextBlock).Text(FText::FromString(ReferencingObjectsList)));
 }
 
@@ -2420,6 +2582,15 @@ void SReferenceViewer::OnInitialAssetRegistrySearchComplete()
 	}
 }
 
+void SReferenceViewer::OnPluginEdited(IPlugin& InPlugin)
+{
+	// The plugin's Verse path may have changed.  Recompute all Verse paths.
+	if (GraphObj)
+	{
+		GraphObj->UpdatePaths();
+	}
+}
+
 void SReferenceViewer::ZoomToFit()
 {
 	if (GraphEditorPtr.IsValid())
@@ -2524,7 +2695,7 @@ void SReferenceViewer::ResolveReferencingProperties() const
 		// If assets need to be loaded in order to resolve properties, let the user know
 		if (!UnloadedAssetsData.IsEmpty())
 		{
-			const EAppReturnType::Type Ret = ShowAssetsNeedsToLoadMessage(UnloadedAssetsData);
+			const EAppReturnType::Type Ret = ShowAssetsNeedsToLoadMessage(UnloadedAssetsData, bShowingContentVersePath);
 			if (Ret == EAppReturnType::Cancel)
 			{
 				return;
@@ -2641,7 +2812,7 @@ void SReferenceViewer::HandleOnSearchTextChanged(const FText& SearchText)
 
 	for (UEdGraphNode_Reference* Node : AllNodes)
 	{
-		if (IsAssetIdentifierPassingSearchTextFilter(Node->GetIdentifier(), SearchWords))
+		if (DoesAssetPassSearchTextFilter(Node->GetIdentifier(), Node->GetAssetData(), bShowingContentVersePath, SearchWords))
 		{
 			GraphEditorPtr->SetNodeSelection(Node, true);
 		}
diff --git a/Engine/Plugins/Editor/AssetManagerEditor/Source/AssetManagerEditor/Private/ReferenceViewer/SReferenceViewer.h b/Engine/Plugins/Editor/AssetManagerEditor/Source/AssetManagerEditor/Private/ReferenceViewer/SReferenceViewer.h
--- a/Engine/Plugins/Editor/AssetManagerEditor/Source/AssetManagerEditor/Private/ReferenceViewer/SReferenceViewer.h
+++ b/Engine/Plugins/Editor/AssetManagerEditor/Source/AssetManagerEditor/Private/ReferenceViewer/SReferenceViewer.h
@@ -9,6 +9,7 @@
 #include "ReferenceViewer/ReferenceViewerSettings.h"
 
 class ICollectionContainer;
+class IPlugin;
 class FUICommandList;
 class SComboButton;
 class SReferenceViewerFilterBar;
@@ -103,6 +104,12 @@ private:
 	/** Gets the text to be displayed in the address bar */
 	FText GetAddressBarText() const;
 
+	/** Gets the summary text to be displayed for multiple FAssetIdentifiers */
+	FText GetIdentifierSummaryText(const TArray<FAssetIdentifier>& Identifiers) const;
+
+	/** Gets the text to be displayed for a FAssetIdentifier */
+	FText GetIdentifierText(const FAssetIdentifier& Identifier) const;
+
 	/** Gets the text to be displayed for warning/status updates */
 	FText GetStatusText() const;
 	
@@ -141,7 +148,7 @@ private:
 
 	void OnEnablePluginFilterChanged(ECheckBoxState NewState);
 	ECheckBoxState IsEnablePluginFilterChecked() const;
-	void PluginFilterAddMenuEntry(FMenuBuilder& MenuBuilder, const FName& PluginName);
+	void PluginFilterAddMenuEntry(FMenuBuilder& MenuBuilder, const FName& PluginName, const FText& Label, const FText& ToolTip);
 	TSharedRef<SWidget> BuildPluginFilterMenu();
 	FText GetPluginComboButtonText() const;
 
@@ -200,7 +207,6 @@ private:
 	bool CanResolveReferencingProperties() const;
 
 	/** Find Path */
-	void OnFindPathAssetSelected();
 	TSharedRef<SWidget> GenerateFindPathAssetPickerMenu();
 	void OnFindPathAssetSelected( const FAssetData& AssetData );
 	void OnFindPathAssetEnterPressed( const TArray<FAssetData>& AssetData );
@@ -213,8 +219,13 @@ private:
 
 	void ReCenterGraphOnNodes(const TSet<UObject*>& Nodes);
 
-	FString GetReferencedObjectsList() const;
-	FString GetReferencingObjectsList() const;
+	enum class EObjectsListType
+	{
+		Referenced,
+		Referencing,
+	};
+
+	FString GetObjectsList(EObjectsListType ObjectsListType) const;
 
 	UObject* GetObjectFromSingleSelectedNode() const;
 	void GetPackageNamesFromSelectedNodes(TSet<FName>& OutNames) const;
@@ -225,6 +236,8 @@ private:
 
 	void OnAssetRegistryChanged(const FAssetData& AssetData);
 	void OnInitialAssetRegistrySearchComplete();
+	void OnPluginEdited(IPlugin& InPlugin);
+
 	EActiveTimerReturnType TriggerZoomToFit(double InCurrentTime, float InDeltaTime);
 private:
 
@@ -286,6 +299,9 @@ private:
 	/** Whether to visually show to the user the option of "Compact Mode" */
 	bool bShowCompactMode;
 
+	/** Whether to show Verse paths */
+	bool bShowingContentVersePath;
+
 	/** A recursion check so as to avoid the rebuild of the graph if we are currently rebuilding the filters */
 	bool bRebuildingFilters;
 
diff --git a/Engine/Plugins/Editor/AssetManagerEditor/Source/AssetManagerEditor/Public/ReferenceViewer/EdGraphNode_Reference.h b/Engine/Plugins/Editor/AssetManagerEditor/Source/AssetManagerEditor/Public/ReferenceViewer/EdGraphNode_Reference.h
--- a/Engine/Plugins/Editor/AssetManagerEditor/Source/AssetManagerEditor/Public/ReferenceViewer/EdGraphNode_Reference.h
+++ b/Engine/Plugins/Editor/AssetManagerEditor/Source/AssetManagerEditor/Public/ReferenceViewer/EdGraphNode_Reference.h
@@ -53,11 +53,13 @@ class ASSETMANAGEREDITOR_API UEdGraphNode_Reference : public UEdGraphNode
 
 	bool IsOverflow() const { return bIsOverflow; }
 
-	FAssetData GetAssetData() const;
+	const FAssetData& GetAssetData() const;
 
 	UEdGraphPin* GetDependencyPin();
 	UEdGraphPin* GetReferencerPin();
 
+	UE_INTERNAL void UpdatePath();
+
 private:
 	void CacheAssetData(const FAssetData& AssetData);
 	void SetupReferenceNode(const FIntPoint& NodeLoc, const TArray<FAssetIdentifier>& NewIdentifiers, const FAssetData& InAssetData, bool bInAllowThumbnail, bool bInIsADuplicate);
diff --git a/Engine/Plugins/Editor/AssetManagerEditor/Source/AssetManagerEditor/Public/ReferenceViewer/EdGraph_ReferenceViewer.h b/Engine/Plugins/Editor/AssetManagerEditor/Source/AssetManagerEditor/Public/ReferenceViewer/EdGraph_ReferenceViewer.h
--- a/Engine/Plugins/Editor/AssetManagerEditor/Source/AssetManagerEditor/Public/ReferenceViewer/EdGraph_ReferenceViewer.h
+++ b/Engine/Plugins/Editor/AssetManagerEditor/Source/AssetManagerEditor/Public/ReferenceViewer/EdGraph_ReferenceViewer.h
@@ -96,8 +96,16 @@ public:
 	 * the presentation or filtering */
 	class UEdGraphNode_Reference* RefilterGraph();
 
+	UE_INTERNAL bool IsShowingContentVersePath() const { return bShowingContentVersePath; }
+	UE_INTERNAL void SetShowingContentVersePath(bool bInShowingContentVersePath);
+	UE_INTERNAL void UpdatePaths();
+
 	using FIsAssetIdentifierPassingSearchFilterCallback = TFunction<bool(const FAssetIdentifier&)>;
-	void SetIsAssetIdentifierPassingSearchFilterCallback(const TOptional<FIsAssetIdentifierPassingSearchFilterCallback>& InIsAssetIdentifierPassingSearchFilterCallback) { IsAssetIdentifierPassingSearchFilterCallback = InIsAssetIdentifierPassingSearchFilterCallback; }
+	UE_DEPRECATED(5.7, "Call SetDoesAssetPassSearchFilterCallback instead.")
+	void SetIsAssetIdentifierPassingSearchFilterCallback(const TOptional<FIsAssetIdentifierPassingSearchFilterCallback>& InIsAssetIdentifierPassingSearchFilterCallback);
+
+	using FDoesAssetPassSearchFilterCallback = TFunction<bool(const FAssetIdentifier&, const FAssetData&)>;
+	void SetDoesAssetPassSearchFilterCallback(const FDoesAssetPassSearchFilterCallback& InDoesAssetPassSearchFilterCallback) { DoesAssetPassSearchFilterCallback = InDoesAssetPassSearchFilterCallback; }
 
 	UE_DEPRECATED(5.6, "Use the ICollectionContainer overload instead.")
 	FName GetCurrentCollectionFilter() const;
@@ -192,7 +200,7 @@ private:
 	void GetSortedLinks(const TArray<FAssetIdentifier>& Identifiers, bool bReferencers, const FAssetManagerDependencyQuery& Query, TMap<FAssetIdentifier, EDependencyPinCategory>& OutLinks) const;
 	bool IsPackageIdentifierPassingFilter(const FAssetIdentifier& InAssetIdentifier) const;
 	bool IsPackageIdentifierPassingPluginFilter(const FAssetIdentifier& InAssetIdentifier) const;
-	bool IsAssetPassingSearchTextFilter(const FAssetIdentifier& InAssetIdentifier) const;
+	bool DoesAssetPassSearchTextFilter(const FAssetIdentifier& InAssetIdentifier, const FAssetData& InAssetData) const;
 
 	UEdGraphNode_Reference* FindPath(const FAssetIdentifier& RootId, const FAssetIdentifier& TargetId);
 	bool FindPath_Recursive(bool bInReferencers, const FAssetIdentifier& InAssetId, const FAssetIdentifier& Target, TMap<FAssetIdentifier, FReferenceNodeInfo>& InNodeInfos, TSet<FAssetIdentifier>& Visited);
@@ -208,11 +216,12 @@ private:
 	TArray<FAssetIdentifier> CurrentGraphRootIdentifiers;
 	FIntPoint CurrentGraphRootOrigin;
 
-	int32 MaxSearchBreadth;
-
 	/** Stores if the breadth limit was reached on the last refilter*/
 	bool bBreadthLimitReached;
 
+	/** Whether to show Verse paths */
+	bool bShowingContentVersePath;
+
 	/** Current collection filter. NAME_None for no filter */
 	TSharedPtr<ICollectionContainer> CurrentCollectionFilterContainer;
 	FName CurrentCollectionFilterName;
@@ -226,14 +235,11 @@ private:
 	/** A set of the unique class types referenced */
 	TSet<FTopLevelAssetPath> CurrentClasses;
 
-	/* This is a convenience toggle to switch between the old & new methods for computing & displaying the graph */
-	bool bUseNodeInfos;
-
 	/* Cached Reference Information used to quickly refilter */
 	TMap<FAssetIdentifier, FReferenceNodeInfo> ReferencerNodeInfos;
 	TMap<FAssetIdentifier, FReferenceNodeInfo> DependencyNodeInfos;
 
-	TOptional<FIsAssetIdentifierPassingSearchFilterCallback> IsAssetIdentifierPassingSearchFilterCallback;
+	FDoesAssetPassSearchFilterCallback DoesAssetPassSearchFilterCallback;
 
 	/** List of packages the current collection filter allows */
 	TSet<FName> CurrentCollectionPackages;
