[Editor-RevisionContrl] Add a diff-against menu entry that allows users to pick any other asset of the same type to diff

#doc https://docs.google.com/document/d/1pTEJBQNCc-V2oKDuttl-K1keZRH5BQHoTU6qmddjMlc/edit?tab=t.0
#tests Used it to diff some different file types

diff --git a/Engine/Plugins/Editor/AssetManagerEditor/Source/AssetManagerEditor/Private/AssetSourceControlContextMenu.cpp b/Engine/Plugins/Editor/AssetManagerEditor/Source/AssetManagerEditor/Private/AssetSourceControlContextMenu.cpp
--- a/Engine/Plugins/Editor/AssetManagerEditor/Source/AssetManagerEditor/Private/AssetSourceControlContextMenu.cpp
+++ b/Engine/Plugins/Editor/AssetManagerEditor/Source/AssetManagerEditor/Private/AssetSourceControlContextMenu.cpp
@@ -13,12 +13,14 @@
 #include "ContentBrowserAssetDataPayload.h"
 #include "ContentBrowserDataMenuContexts.h"
 #include "ContentBrowserMenuContexts.h"
+#include "ContentBrowserModule.h"
 #include "FileHelpers.h"
 #include "HAL/IConsoleManager.h"
 #include "HAL/PlatformApplicationMisc.h"
 #include "HAL/PlatformFile.h"
 #include "HAL/PlatformFileManager.h"
 #include "IAssetTools.h"
+#include "IContentBrowserSingleton.h"
 #include "ISourceControlModule.h"
 #include "ISourceControlRevision.h"
 #include "Misc/MessageDialog.h"
@@ -131,6 +133,7 @@ public:
 private:
 	bool AddSourceControlMenuOptions(FToolMenuSection& InSection);
 	void FillSourceControlSubMenu(UToolMenu* Menu);
+	void FillDiffAgainstSubMenu(UToolMenu* Menu);
 
 	// Todo: Make TriggerCheckAtNumActions default value into a user option
 
@@ -407,6 +410,17 @@ bool FAssetSourceControlContextMenuState::AddSourceControlMenuOptions(FToolMenuS
 			)
 		);
 	}
+	else if (SelectedAssets.Num() == 1 && !bContainsFolders)
+	{
+		InSection.AddSubMenu(
+			"DiffAgainst",
+			LOCTEXT("DiffAgainst", "Diff Against..."),
+			LOCTEXT("DiffAgainstTooltip", "Pick another asset of the same type to compare."),
+			FNewToolMenuDelegate::CreateSP(this, &FAssetSourceControlContextMenuState::FillDiffAgainstSubMenu),
+			false,
+			FSlateIcon(FRevisionControlStyleManager::GetStyleSetName(), "RevisionControl.Actions.Diff")
+		);
+	}
 
 	return true;
 }
@@ -650,6 +664,62 @@ void FAssetSourceControlContextMenuState::FillSourceControlSubMenu(UToolMenu* Me
 	}
 }
 
+void FAssetSourceControlContextMenuState::FillDiffAgainstSubMenu(UToolMenu* Menu)
+{
+	if (SelectedAssets.Num() != 1)
+	{
+		return;
+	}
+
+	FAssetPickerConfig AssetPickerConfig;
+	AssetPickerConfig.SelectionMode = ESelectionMode::Single;
+	AssetPickerConfig.InitialThumbnailSize = EThumbnailSize::Tiny;
+	AssetPickerConfig.InitialAssetViewType = EAssetViewType::List;
+	AssetPickerConfig.bAllowDragging = false;
+	AssetPickerConfig.Filter.ClassPaths.Add(SelectedAssets[0].GetClass()->GetClassPathName());
+	AssetPickerConfig.OnAssetSelected = FOnAssetSelected::CreateLambda([this](const FAssetData& PickedAsset)
+		{
+			if (SelectedAssets.Num() == 1)
+			{
+				UObject* FirstObjectSelected = SelectedAssets[0].GetAsset();
+				UObject* SecondObjectSelected = PickedAsset.GetAsset();
+
+				if ((FirstObjectSelected != NULL) && (SecondObjectSelected != NULL))
+				{
+					if (FirstObjectSelected == SecondObjectSelected)
+					{
+						FMessageDialog::Open(EAppMsgType::Ok, LOCTEXT("SourceControlDiffSameAssets", "Can't diff an asset against itself."));
+					}
+					else
+					{
+						// Load the asset registry module
+						FAssetToolsModule& AssetToolsModule = FModuleManager::GetModuleChecked<FAssetToolsModule>("AssetTools");
+
+						FRevisionInfo CurrentRevision;
+						CurrentRevision.Revision = TEXT("");
+
+						// Need to dismiss context menu first or the diff window will try to attach to it
+						FSlateApplication::Get().DismissAllMenus();
+
+						AssetToolsModule.Get().DiffAssets(FirstObjectSelected, SecondObjectSelected, CurrentRevision, CurrentRevision);
+					}
+				}
+			}
+		});
+
+	TSharedRef<SWidget> AssetPicker = 
+		SNew(SBox)
+	.WidthOverride(300.f)
+	.HeightOverride(300.f)
+	[
+		FModuleManager::LoadModuleChecked<FContentBrowserModule>("ContentBrowser").Get().CreateAssetPicker(AssetPickerConfig)
+	];
+
+	FToolMenuSection& Section = Menu->AddSection("AssetDiffAgainstPicker", LOCTEXT("AssetDiffAgainstPickerHeading", "Diff Against..."));
+
+	Section.AddEntry(FToolMenuEntry::InitWidget("DiffAgainst", AssetPicker, {}));
+}
+
 FExecuteAction FAssetSourceControlContextMenuState::ExecutionCheck(FExecuteAction&& InAction, int32 TriggerCheckAtNumActions) const
 {
 	return FExecuteAction::CreateSPLambda(this, [this, TriggerCheckAtNumActions, Action = MoveTemp(InAction)]()
