Expose the ability to register commands to expand the right-click menu on nodes in the AssetReferenceViewer

-Adds OnRegisterAdditionalReferenceViewerCommands to IAssetManagerEditorModule, should be used to register delegates to append commands to be added to the right-click menu on nodes in AssetReferenceViewer.

-Adds GetRegisteredAdditionalReferenceViewerCommands to IAssetManagerEditorModule, which is used to get command list by ReferenceViewerSchema

-Adds GetReferenceViewerUISelectedNodesAssetData to IAssetManagerEditorModule, which can be used to get AssetData for the currently selected nodes in the active tab. Intended for use by registered additional commands that need to get currently selected nodes.

#changelist validated

diff --git a/Engine/Plugins/Editor/AssetManagerEditor/Source/AssetManagerEditor/Private/AssetManagerEditorModule.cpp b/Engine/Plugins/Editor/AssetManagerEditor/Source/AssetManagerEditor/Private/AssetManagerEditorModule.cpp
--- a/Engine/Plugins/Editor/AssetManagerEditor/Source/AssetManagerEditor/Private/AssetManagerEditorModule.cpp
+++ b/Engine/Plugins/Editor/AssetManagerEditor/Source/AssetManagerEditor/Private/AssetManagerEditorModule.cpp
@@ -332,6 +332,7 @@ public:
 	virtual FCanOpenReferenceViewerUI& OnCanOpenReferenceViewerUI() override;
 	virtual void OpenReferenceViewerUI(const TArray<FAssetIdentifier> SelectedIdentifiers, const FReferenceViewerParams ReferenceViewerParams = FReferenceViewerParams()) override;
 	virtual void OpenReferenceViewerUI(const TArray<FName> SelectedPackages, const FReferenceViewerParams ReferenceViewerParams = FReferenceViewerParams()) override;
+	virtual bool GetReferenceViewerUISelectedNodesAssetData(TArray<FAssetData>& OutAssetDatas) override;
 	virtual void OpenSizeMapUI(TArray<FAssetIdentifier> SelectedIdentifiers) override;
 	virtual void OpenSizeMapUI(TArray<FName> SelectedPackages) override;	
 	virtual void OpenShaderCookStatistics(TArray<FName> SelectedIdentifiers) override;
@@ -348,6 +349,10 @@ public:
 	virtual bool FilterAssetIdentifiersForCurrentRegistrySource(TArray<FAssetIdentifier>& AssetIdentifiers, const FAssetManagerDependencyQuery& DependencyQuery = FAssetManagerDependencyQuery::None(), bool bForwardDependency = true) override;
 	virtual bool WriteCollection(FName CollectionName, ECollectionShareType::Type ShareType, const TArray<FName>& PackageNames, bool bShowFeedback) override;
 	virtual bool WriteCollection(ICollectionContainer& CollectionContainer, FName CollectionName, ECollectionShareType::Type ShareType, const TArray<FName>& PackageNames, bool bShowFeedback) override;
+	
+	virtual void AppendRegisteredCommandsToCommandList(TSharedRef<FUICommandList> ListToAppendTo) override;
+	virtual FRegisterAdditionalReferenceViewerCommands& OnRegisterAdditionalReferenceViewerCommands() override { return RegisteredCommands; }
+	virtual const TArray<TSharedPtr<FUICommandInfo>>& GetRegisteredAdditionalReferenceViewerCommands() override { return AdditionalCommandsInfoList; }
 
 private:
 
@@ -426,6 +431,11 @@ private:
 
 	FCanOpenReferenceViewerUI CanOpenReferenceViewerUIDelegate;
 
+	/**Delegate for registering additional Commands */
+	FRegisterAdditionalReferenceViewerCommands RegisteredCommands;
+
+	TArray<TSharedPtr<FUICommandInfo>> AdditionalCommandsInfoList;
+
 	/**
 	 * When currently selected assets are changed for an already opened Reference Viewer Tab,
 	 * this function can be called to properly update Tab Hash and Label
@@ -963,6 +973,40 @@ void FAssetManagerEditorModule::OpenReferenceViewerUI(const TArray<FName> Select
 	OpenReferenceViewerUI(Identifiers, ReferenceViewerParams);
 }
 
+bool FAssetManagerEditorModule::GetReferenceViewerUISelectedNodesAssetData(TArray<FAssetData>& OutAssetDatas)
+{
+	bool ValidAssetDataFound = false;
+	for (auto TabPair : ReferenceViewerTabs)
+	{
+		if (TabPair.Value.IsValid())
+		{
+			const TSharedPtr<SDockTab> Tab = TabPair.Value.Pin();
+
+			// we only want to do this on the actual active tab
+			if (Tab->IsForeground())
+			{
+				TSharedPtr<SWidget> TabContent = Tab->GetContent();
+
+				if (TabContent->GetWidgetClass().GetWidgetType() != SReferenceViewer::StaticWidgetClass().GetWidgetType())
+				{
+					continue;
+				}
+
+				TSharedPtr<SReferenceViewer> ActiveReferenceViewer = StaticCastSharedPtr<SReferenceViewer>(TabContent);
+
+				TArray<FAssetData> SelectedAssets;
+				ActiveReferenceViewer->GetSelectedNodeAssetData(SelectedAssets);
+
+				// just in case we can have two valid tabs displayed at once, writing this in such a way that it should still work
+				OutAssetDatas.Append(SelectedAssets);
+				ValidAssetDataFound = true;
+			}
+		}
+	}
+
+	return ValidAssetDataFound;
+}
+
 void FAssetManagerEditorModule::OpenSizeMapUI(TArray<FName> SelectedPackages)
 {
 	TArray<FAssetIdentifier> Identifiers;
@@ -2859,4 +2903,10 @@ bool FAssetManagerEditorModule::WriteCollection(ICollectionContainer& Collection
 	return bSuccess;
 }
 
-#undef LOCTEXT_NAMESPACE
+void FAssetManagerEditorModule::AppendRegisteredCommandsToCommandList(TSharedRef<FUICommandList> ListToAppendTo)
+{
+	// we maintain a separate commands list so we can register these to the menu in the Reference Viewer Schema
+	RegisteredCommands.Broadcast(ListToAppendTo, AdditionalCommandsInfoList);
+}
+
+#undef LOCTEXT_NAMESPACE
\ No newline at end of file
diff --git a/Engine/Plugins/Editor/AssetManagerEditor/Source/AssetManagerEditor/Private/ReferenceViewer/ReferenceViewerSchema.cpp b/Engine/Plugins/Editor/AssetManagerEditor/Source/AssetManagerEditor/Private/ReferenceViewer/ReferenceViewerSchema.cpp
--- a/Engine/Plugins/Editor/AssetManagerEditor/Source/AssetManagerEditor/Private/ReferenceViewer/ReferenceViewerSchema.cpp
+++ b/Engine/Plugins/Editor/AssetManagerEditor/Source/AssetManagerEditor/Private/ReferenceViewer/ReferenceViewerSchema.cpp
@@ -187,6 +187,25 @@ void UReferenceViewerSchema::GetContextMenuActions(UToolMenu* Menu, UGraphNodeCo
 		ViewAssetAuditEntry.Name = TEXT("ContextMenu");
 		Section.AddEntry(ViewAssetAuditEntry);
 	}
+
+	// Our additional registered entries
+	if (IAssetManagerEditorModule::IsAvailable())
+	{
+		const TArray<TSharedPtr<FUICommandInfo>> &CommandInfoList = IAssetManagerEditorModule::Get().GetRegisteredAdditionalReferenceViewerCommands();
+
+		if (CommandInfoList.Num() > 0)
+		{
+			FToolMenuSection& Section = Menu->AddSection(TEXT("Extensions"), NSLOCTEXT("ReferenceViewerSchema", "ExtensionsSectionLabel", "Extensions"));
+
+			for (TSharedPtr<FUICommandInfo> Entry : CommandInfoList)
+			{
+				if (Entry.IsValid())
+				{
+					Section.AddMenuEntry(Entry);
+				}
+			}
+		}
+	}
 }
 
 FLinearColor UReferenceViewerSchema::GetPinTypeColor(const FEdGraphPinType& PinType) const
diff --git a/Engine/Plugins/Editor/AssetManagerEditor/Source/AssetManagerEditor/Private/ReferenceViewer/SReferenceViewer.cpp b/Engine/Plugins/Editor/AssetManagerEditor/Source/AssetManagerEditor/Private/ReferenceViewer/SReferenceViewer.cpp
--- a/Engine/Plugins/Editor/AssetManagerEditor/Source/AssetManagerEditor/Private/ReferenceViewer/SReferenceViewer.cpp
+++ b/Engine/Plugins/Editor/AssetManagerEditor/Source/AssetManagerEditor/Private/ReferenceViewer/SReferenceViewer.cpp
@@ -793,6 +793,29 @@ EActiveTimerReturnType SReferenceViewer::TriggerZoomToFit(double InCurrentTime,
 	return EActiveTimerReturnType::Stop;
 }
 
+void SReferenceViewer::GetSelectedNodeAssetData(TArray<FAssetData>& OutAssetData) const
+{
+	TSet<UObject*> SelectedNodes = GraphEditorPtr->GetSelectedNodes();
+
+	Algo::TransformIf(SelectedNodes, OutAssetData,
+		[](UObject* AssetToCheck) -> bool
+		{
+			// might need to do more just doing what is done above for now
+			if (AssetToCheck != nullptr && AssetToCheck->IsA<UEdGraphNode_Reference>())
+			{
+				return Cast<UEdGraphNode_Reference>(AssetToCheck)->GetAssetData().IsValid();
+			}
+			else
+			{
+				return false;
+			}
+		},
+		[](UObject* NodeToGetDataFrom) -> FAssetData
+		{
+			return Cast<UEdGraphNode_Reference>(NodeToGetDataFrom)->GetAssetData();
+		});
+}
+
 void SReferenceViewer::SetCurrentRegistrySource(const FAssetManagerEditorRegistrySource* RegistrySource)
 {
 	RebuildGraph();
@@ -2067,6 +2090,8 @@ void SReferenceViewer::RegisterActions()
 				}
 		}),
 		FCanExecuteAction::CreateSP(this, &SReferenceViewer::HasAtLeastOneRealNodeSelected));
+
+	IAssetManagerEditorModule::Get().AppendRegisteredCommandsToCommandList(ReferenceViewerActions.ToSharedRef());
 }
 
 void SReferenceViewer::ShowSelectionInContentBrowser()
diff --git a/Engine/Plugins/Editor/AssetManagerEditor/Source/AssetManagerEditor/Private/ReferenceViewer/SReferenceViewer.h b/Engine/Plugins/Editor/AssetManagerEditor/Source/AssetManagerEditor/Private/ReferenceViewer/SReferenceViewer.h
--- a/Engine/Plugins/Editor/AssetManagerEditor/Source/AssetManagerEditor/Private/ReferenceViewer/SReferenceViewer.h
+++ b/Engine/Plugins/Editor/AssetManagerEditor/Source/AssetManagerEditor/Private/ReferenceViewer/SReferenceViewer.h
@@ -63,6 +63,7 @@ public:
 
 	FOnReferenceViewerSelectionChanged& OnReferenceViewerSelectionChanged() { return OnReferenceViewerSelectionChangedDelegate; }
 
+	void GetSelectedNodeAssetData(TArray<FAssetData>& OutAssetData) const;
 private:
 
 	/** Call after a structural change is made that causes the graph to be recreated */
diff --git a/Engine/Plugins/Editor/AssetManagerEditor/Source/AssetManagerEditor/Public/AssetManagerEditorModule.h b/Engine/Plugins/Editor/AssetManagerEditor/Source/AssetManagerEditor/Public/AssetManagerEditorModule.h
--- a/Engine/Plugins/Editor/AssetManagerEditor/Source/AssetManagerEditor/Public/AssetManagerEditorModule.h
+++ b/Engine/Plugins/Editor/AssetManagerEditor/Source/AssetManagerEditor/Public/AssetManagerEditorModule.h
@@ -12,7 +12,8 @@
 #include "CollectionManagerTypes.h"
 
 class ICollectionContainer;
-
+class FUICommandList;
+class FUICommandInfo;
 /**
  * AssetRegistry Dependency Category + QueryFlags, to pass around the two enums in a single variable.
  */
@@ -426,6 +427,9 @@ public:
 	virtual void OpenReferenceViewerUI(const TArray<FAssetIdentifier> SelectedIdentifiers, const FReferenceViewerParams ReferenceViewerParams = FReferenceViewerParams()) = 0;
 	virtual void OpenReferenceViewerUI(const TArray<FName> SelectedPackages, const FReferenceViewerParams ReferenceViewerParams = FReferenceViewerParams()) = 0;
 
+	/** Returns true if Reference Viewer UI is available and has selected nodes. If true, will also append the asset data for the selected nodes into the passed array*/
+	virtual bool GetReferenceViewerUISelectedNodesAssetData(TArray<FAssetData>& OutAssetDatas) = 0;
+
 	/** Returns false if should not open reference viewer for items */
 	DECLARE_DELEGATE_RetVal_TwoParams(bool, FCanOpenReferenceViewerUI, const TArray<FAssetIdentifier>&, FText*);
 	virtual FCanOpenReferenceViewerUI& OnCanOpenReferenceViewerUI() = 0;
@@ -436,4 +440,14 @@ public:
 
 	/** Open the Shader cook stats */
 	virtual void OpenShaderCookStatistics(TArray<FName> SelectedIdentifiers) = 0;
+
+	/** Called by SReferenceViewer to update given list of commands with any of the registered delegates */
+	virtual void AppendRegisteredCommandsToCommandList(TSharedRef<FUICommandList> ListToAppendTo) = 0;
+
+	/** Passes commandlist and command info array to be updated by any of the registered delegates. Broadcast from AppendRegisteredCommandsToCommandList */
+	DECLARE_MULTICAST_DELEGATE_TwoParams(FRegisterAdditionalReferenceViewerCommands, TSharedRef<FUICommandList>, TArray<TSharedPtr<FUICommandInfo>>&);
+	virtual FRegisterAdditionalReferenceViewerCommands &OnRegisterAdditionalReferenceViewerCommands() = 0;
+
+	/** Returns our list of command info populated by OnRegisterAdditionalReferenceViewerCommands. Needed to be tracked separately from command list to populate commands in menu ui */
+	virtual const TArray<TSharedPtr<FUICommandInfo>>& GetRegisteredAdditionalReferenceViewerCommands() = 0;
 };
