[Reference Viewer] Graph nodes not selected when filtering using the search text

diff --git a/Engine/Plugins/Editor/AssetManagerEditor/Source/AssetManagerEditor/Private/ReferenceViewer/SReferenceViewer.cpp b/Engine/Plugins/Editor/AssetManagerEditor/Source/AssetManagerEditor/Private/ReferenceViewer/SReferenceViewer.cpp
--- a/Engine/Plugins/Editor/AssetManagerEditor/Source/AssetManagerEditor/Private/ReferenceViewer/SReferenceViewer.cpp
+++ b/Engine/Plugins/Editor/AssetManagerEditor/Source/AssetManagerEditor/Private/ReferenceViewer/SReferenceViewer.cpp
@@ -2825,6 +2825,9 @@ void SReferenceViewer::HandleOnSearchTextChanged(const FText& SearchText)
 
 	if (SearchText.IsEmpty())
 	{
+		constexpr bool bOnlySelection = false;
+		// Zoom back to show the entire graph if nothing is selected
+		GraphEditorPtr->ZoomToFit(bOnlySelection);
 		return;
 	}
 
@@ -2839,12 +2842,14 @@ void SReferenceViewer::HandleOnSearchTextChanged(const FText& SearchText)
 	{
 		if (DoesAssetPassSearchTextFilter(Node->GetIdentifier(), Node->GetAssetData(), bShowingContentVersePath, SearchWords))
 		{
-			GraphEditorPtr->SetNodeSelection(Node, true);
+			constexpr bool bSelect = true;
+			GraphEditorPtr->SetNodeSelection(Node, bSelect);
 		}
 	}
 
-	// Ensure shown nodes are up-to-date with search filter
-	GraphObj->RefilterGraph();
+	constexpr bool bOnlySelection = true;
+	// Zoom to fit the select nodes. Also ensures the graph is up to date
+	GraphEditorPtr->ZoomToFit(bOnlySelection);
 }
 
 void SReferenceViewer::HandleOnSearchTextCommitted(const FText& SearchText, ETextCommit::Type CommitType)
@@ -2857,13 +2862,15 @@ void SReferenceViewer::HandleOnSearchTextCommitted(const FText& SearchText, ETex
 	if (CommitType == ETextCommit::OnCleared)
 	{
 		GraphEditorPtr->ClearSelectionSet();
+
+		constexpr bool bOnlySelection = true;
+		GraphEditorPtr->ZoomToFit(bOnlySelection);
 	}
 	else if (CommitType == ETextCommit::OnEnter)
 	{
 		HandleOnSearchTextChanged(SearchBox->GetText());
 	}
-	
-	GraphEditorPtr->ZoomToFit(true);
+
 	FSlateApplication::Get().SetKeyboardFocus(GraphEditorPtr, EFocusCause::SetDirectly);
 }
 
