Add API to be able to show asset source control context menu from anywhere in the editor.

#rnx

diff --git a/Engine/Plugins/Editor/AssetManagerEditor/Source/AssetManagerEditor/Private/AssetManagerEditorModule.cpp b/Engine/Plugins/Editor/AssetManagerEditor/Source/AssetManagerEditor/Private/AssetManagerEditorModule.cpp
--- a/Engine/Plugins/Editor/AssetManagerEditor/Source/AssetManagerEditor/Private/AssetManagerEditorModule.cpp
+++ b/Engine/Plugins/Editor/AssetManagerEditor/Source/AssetManagerEditor/Private/AssetManagerEditorModule.cpp
@@ -353,6 +353,7 @@ public:
 	virtual void AppendRegisteredCommandsToCommandList(TSharedRef<FUICommandList> ListToAppendTo) override;
 	virtual FRegisterAdditionalReferenceViewerCommands& OnRegisterAdditionalReferenceViewerCommands() override { return RegisteredCommands; }
 	virtual const TArray<TSharedPtr<FUICommandInfo>>& GetRegisteredAdditionalReferenceViewerCommands() override { return AdditionalCommandsInfoList; }
+	virtual TSharedRef<SWidget> CreateAssetSourceControlContextMenu(const TArray<FAssetData>& SelectedAssets) override;
 
 private:
 
@@ -2909,4 +2910,38 @@ void FAssetManagerEditorModule::AppendRegisteredCommandsToCommandList(TSharedRef
 	RegisteredCommands.Broadcast(ListToAppendTo, AdditionalCommandsInfoList);
 }
 
+TSharedRef<SWidget> FAssetManagerEditorModule::CreateAssetSourceControlContextMenu(const TArray<FAssetData>& SelectedAssets)
+{
+	UToolMenus* ToolMenus = UToolMenus::Get();
+	//Find or add the tool menu extension for this case
+	UToolMenu* ContextMenu = ToolMenus->ExtendMenu("GenericContextMenu.AssetSourceControl");
+	
+	//Find or add the section. We need at least one section in the context menu. We do not need any label since there is only one section.
+	FToolMenuSection& Section = ContextMenu->FindOrAddSection("Section");
+
+	//If we did not already add a UContentBrowserAssetContextMenuContext context object we must create one and add it.
+	UContentBrowserAssetContextMenuContext* ContentBrowserAssetContextMenuContext = Section.Context.FindContext<UContentBrowserAssetContextMenuContext>();
+	if(!ContentBrowserAssetContextMenuContext)
+	{
+		//Create a content browser asset context menu context. This is required to use the built in source control context menu.
+		//The goal of this context is to pass the asset list on which we want to do source control operation.
+		ContentBrowserAssetContextMenuContext = NewObject<UContentBrowserAssetContextMenuContext>();
+		ContentBrowserAssetContextMenuContext->bCanBeModified = true;
+		//Add the context to the section
+		Section.Context.AddObject(ContentBrowserAssetContextMenuContext);
+	}
+
+	//Set the selected assets
+	if(ensure(ContentBrowserAssetContextMenuContext))
+	{
+		ContentBrowserAssetContextMenuContext->SelectedAssets = SelectedAssets;
+	}
+
+	//Build the generic source control context menu
+	AssetSourceControlContextMenu.MakeContextMenu(Section);
+
+	//Generate the widget use to show the menu
+	return ToolMenus->GenerateWidget(ContextMenu);
+}
+
 #undef LOCTEXT_NAMESPACE
\ No newline at end of file
diff --git a/Engine/Plugins/Editor/AssetManagerEditor/Source/AssetManagerEditor/Public/AssetManagerEditorModule.h b/Engine/Plugins/Editor/AssetManagerEditor/Source/AssetManagerEditor/Public/AssetManagerEditorModule.h
--- a/Engine/Plugins/Editor/AssetManagerEditor/Source/AssetManagerEditor/Public/AssetManagerEditorModule.h
+++ b/Engine/Plugins/Editor/AssetManagerEditor/Source/AssetManagerEditor/Public/AssetManagerEditorModule.h
@@ -450,4 +450,7 @@ public:
 
 	/** Returns our list of command info populated by OnRegisterAdditionalReferenceViewerCommands. Needed to be tracked separately from command list to populate commands in menu ui */
 	virtual const TArray<TSharedPtr<FUICommandInfo>>& GetRegisteredAdditionalReferenceViewerCommands() = 0;
+	
+	/* Return a UToolsMenu generated widget containing the asset source control context menu. */
+	virtual TSharedRef<SWidget> CreateAssetSourceControlContextMenu(const TArray<FAssetData>& SelectedAssets) = 0;
 };
