[UE] Fix crash when preparing context menu in content browser

#tests Made sure menu still worked, not sure the exact case that caused the crash
#rnx

diff --git a/Engine/Plugins/Editor/AssetManagerEditor/Source/AssetManagerEditor/Private/AssetSourceControlContextMenu.cpp b/Engine/Plugins/Editor/AssetManagerEditor/Source/AssetManagerEditor/Private/AssetSourceControlContextMenu.cpp
--- a/Engine/Plugins/Editor/AssetManagerEditor/Source/AssetManagerEditor/Private/AssetSourceControlContextMenu.cpp
+++ b/Engine/Plugins/Editor/AssetManagerEditor/Source/AssetManagerEditor/Private/AssetSourceControlContextMenu.cpp
@@ -410,7 +410,7 @@ bool FAssetSourceControlContextMenuState::AddSourceControlMenuOptions(FToolMenuS
 			)
 		);
 	}
-	else if (SelectedAssets.Num() == 1 && !bContainsFolders)
+	else if (SelectedAssets.Num() == 1 && SelectedAssets[0].GetClass() && !bContainsFolders)
 	{
 		InSection.AddSubMenu(
 			"DiffAgainst",
@@ -666,7 +666,7 @@ void FAssetSourceControlContextMenuState::FillSourceControlSubMenu(UToolMenu* Me
 
 void FAssetSourceControlContextMenuState::FillDiffAgainstSubMenu(UToolMenu* Menu)
 {
-	if (SelectedAssets.Num() != 1)
+	if (SelectedAssets.Num() != 1 || SelectedAssets[0].GetClass() == nullptr)
 	{
 		return;
 	}
