[Reference Viewer] Add Revision Control Icon to nodes

- Enables asset status thumbnail overlay
- Adds base class for custom Reference Viewer filters
- Adds custom "Check Out" filter
- Stores filter status in config file

diff --git a/Engine/Plugins/Editor/AssetManagerEditor/Source/AssetManagerEditor/Private/ReferenceViewer/ReferenceViewerFilters.cpp b/Engine/Plugins/Editor/AssetManagerEditor/Source/AssetManagerEditor/Private/ReferenceViewer/ReferenceViewerFilters.cpp
new file mode 100644
--- /dev/null
+++ b/Engine/Plugins/Editor/AssetManagerEditor/Source/AssetManagerEditor/Private/ReferenceViewer/ReferenceViewerFilters.cpp
@@ -0,0 +1,88 @@
+// Copyright Epic Games, Inc. All Rights Reserved.
+
+#include "ReferenceViewerFilters.h"
+
+#include "ISourceControlModule.h"
+#include "ISourceControlProvider.h"
+#include "ReferenceViewer/EdGraph_ReferenceViewer.h"
+#include "SourceControlOperations.h"
+
+#define LOCTEXT_NAMESPACE "ReferenceViewerFilters"
+
+FLinearColor FReferenceViewerFilter::GetColor() const
+{
+	return FLinearColor(0.6f, 0.6f, 0.6f, 1);
+}
+
+FName FReferenceViewerFilter::GetIconName() const
+{
+	return NAME_None;
+}
+
+FReferenceViewerFilter_ShowCheckedOut::FReferenceViewerFilter_ShowCheckedOut(const TSharedPtr<FFilterCategory>& InCategory)
+	: FReferenceViewerFilter(InCategory)
+{
+}
+
+FString FReferenceViewerFilter_ShowCheckedOut::GetName() const
+{
+	return TEXT("RefViewer_CheckedOut");
+}
+
+FText FReferenceViewerFilter_ShowCheckedOut::GetDisplayName() const
+{
+	return LOCTEXT("ReferenceViewerFilter_ShowCheckedOut", "Checked Out");
+}
+
+FText FReferenceViewerFilter_ShowCheckedOut::GetToolTipText() const
+{
+	return LOCTEXT("ReferenceViewerFilter_ShowCheckedOut", "Allow display of Checked Out Items.");
+}
+
+void FReferenceViewerFilter_ShowCheckedOut::ActiveStateChanged(bool bInActive)
+{
+	if (bInActive)
+	{
+		RequestStatus();
+	}
+}
+
+bool FReferenceViewerFilter_ShowCheckedOut::PassesFilter(FReferenceNodeInfo& InItem) const
+{
+	if (!ISourceControlModule::Get().IsEnabled())
+	{
+		return false;
+	}
+
+	FSourceControlStatePtr SourceControlState =
+		ISourceControlModule::Get().GetProvider().GetState(InItem.AssetData.GetPackage(), EStateCacheUsage::Use);
+
+	return SourceControlState.IsValid() && (SourceControlState->IsCheckedOut() || SourceControlState->IsAdded());
+}
+
+void FReferenceViewerFilter_ShowCheckedOut::RequestStatus()
+{
+	ISourceControlProvider& SourceControlProvider = ISourceControlModule::Get().GetProvider();
+	if (ISourceControlModule::Get().IsEnabled())
+	{
+		// Request the opened files at filter construction time to make sure checked out files have the correct state for the filter
+		TSharedRef<FUpdateStatus, ESPMode::ThreadSafe> UpdateStatusOperation =
+			ISourceControlOperation::Create<FUpdateStatus>();
+		UpdateStatusOperation->SetGetOpenedOnly(true);
+		SourceControlProvider.Execute(
+			UpdateStatusOperation,
+			EConcurrency::Asynchronous,
+			FSourceControlOperationComplete::CreateSP(this, &FReferenceViewerFilter_ShowCheckedOut::SourceControlOperationComplete)
+		);
+	}
+}
+
+void FReferenceViewerFilter_ShowCheckedOut::SourceControlOperationComplete(
+	const FSourceControlOperationRef& InOperation, ECommandResult::Type InResult
+)
+{
+	BroadcastChangedEvent();
+}
+
+
+#undef LOCTEXT_NAMESPACE
diff --git a/Engine/Plugins/Editor/AssetManagerEditor/Source/AssetManagerEditor/Private/ReferenceViewer/ReferenceViewerFilters.h b/Engine/Plugins/Editor/AssetManagerEditor/Source/AssetManagerEditor/Private/ReferenceViewer/ReferenceViewerFilters.h
new file mode 100644
--- /dev/null
+++ b/Engine/Plugins/Editor/AssetManagerEditor/Source/AssetManagerEditor/Private/ReferenceViewer/ReferenceViewerFilters.h
@@ -0,0 +1,75 @@
+// Copyright Epic Games, Inc. All Rights Reserved.
+
+#pragma once
+
+#include "ISourceControlOperation.h"
+#include "ISourceControlProvider.h"
+#include "Filters/FilterBase.h"
+
+struct FReferenceNodeInfo;
+
+class FReferenceViewerFilter : public FFilterBase<FReferenceNodeInfo&>
+{
+public:
+	FReferenceViewerFilter(TSharedPtr<FFilterCategory> InCategory)
+		: FFilterBase<FReferenceNodeInfo&>(MoveTemp(InCategory))
+	{}
+
+	//~ Begin FFilterBase
+	/** Returns the color this filter button will be when displayed as a button */
+	virtual FLinearColor GetColor() const override;
+
+	/** Returns the name of the icon to use in menu entries */
+	virtual FName GetIconName() const override;
+
+	/** Returns true if the filter should be in the list when disabled and not in the list when enabled */
+	virtual bool IsInverseFilter() const override
+	{
+		return false;
+	}
+
+	/** Notification that the filter became active or inactive */
+	virtual void ActiveStateChanged(bool bActive) override {}
+
+	/** Called when the right-click context menu is being built for this filter */
+	virtual void ModifyContextMenu(FMenuBuilder& MenuBuilder) override {}
+
+	/** Called when the state of a particular Content Browser is being saved to INI */
+	virtual void SaveSettings(const FString& IniFilename, const FString& IniSection, const FString& SettingsString) const override {}
+
+	/** Called when the state of a particular Content Browser is being loaded from INI */
+	virtual void LoadSettings(const FString& IniFilename, const FString& IniSection, const FString& SettingsString) override {}
+	//~ End FFilterBase
+};
+
+/**
+ * A Filter to show nodes checked out by revision control
+ */
+class FReferenceViewerFilter_ShowCheckedOut final
+	: public FReferenceViewerFilter
+	, public TSharedFromThis<FReferenceViewerFilter_ShowCheckedOut>
+{
+public:
+	FReferenceViewerFilter_ShowCheckedOut(const TSharedPtr<FFilterCategory>& InCategory);
+
+	/** Returns the system name for this filter */
+	virtual FString GetName() const override;
+
+	/** Returns the human-readable name for this filter */
+	virtual FText GetDisplayName() const override;
+
+	/** Returns the tooltip for this filter, shown in the filters menu */
+	virtual FText GetToolTipText() const override;
+
+	/** Notification that the filter became active or inactive */
+	virtual void ActiveStateChanged(bool bInActive) override;
+
+	virtual bool PassesFilter(FReferenceNodeInfo& InItem) const override;
+
+private:
+	/** Request the source control status for this filter */
+	void RequestStatus();
+
+	/** Callback when source control operation has completed */
+	void SourceControlOperationComplete(const FSourceControlOperationRef& InOperation, ECommandResult::Type InResult);
+};
diff --git a/Engine/Plugins/Editor/AssetManagerEditor/Source/AssetManagerEditor/Private/ReferenceViewer/ReferenceViewerSettings.cpp b/Engine/Plugins/Editor/AssetManagerEditor/Source/AssetManagerEditor/Private/ReferenceViewer/ReferenceViewerSettings.cpp
--- a/Engine/Plugins/Editor/AssetManagerEditor/Source/AssetManagerEditor/Private/ReferenceViewer/ReferenceViewerSettings.cpp
+++ b/Engine/Plugins/Editor/AssetManagerEditor/Source/AssetManagerEditor/Private/ReferenceViewer/ReferenceViewerSettings.cpp
@@ -263,3 +263,14 @@ void UReferenceViewerSettings::SetUserFilters(TArray<FilterState>& InFilters)
 	SaveConfig();
 }
 
+const TMap<FString, bool>& UReferenceViewerSettings::GetCustomFilters() const
+{
+	return CustomFilters;
+}
+
+void UReferenceViewerSettings::SetCustomFilters(const TMap<FString, bool>& InCustomFilters)
+{
+	CustomFilters = InCustomFilters;
+	SaveConfig();
+}
+
diff --git a/Engine/Plugins/Editor/AssetManagerEditor/Source/AssetManagerEditor/Private/ReferenceViewer/SReferenceNode.cpp b/Engine/Plugins/Editor/AssetManagerEditor/Source/AssetManagerEditor/Private/ReferenceViewer/SReferenceNode.cpp
--- a/Engine/Plugins/Editor/AssetManagerEditor/Source/AssetManagerEditor/Private/ReferenceViewer/SReferenceNode.cpp
+++ b/Engine/Plugins/Editor/AssetManagerEditor/Source/AssetManagerEditor/Private/ReferenceViewer/SReferenceNode.cpp
@@ -119,6 +119,7 @@ void SReferenceNode::UpdateGraphNode()
 		ThumbnailConfig.bAllowFadeIn = RefGraphNode->UsesThumbnail();
 		ThumbnailConfig.bForceGenericThumbnail = !RefGraphNode->UsesThumbnail();
 		ThumbnailConfig.AssetTypeColorOverride = FLinearColor::Transparent;
+		ThumbnailConfig.AllowAssetStatusThumbnailOverlay = true;
 
 		MainNodeWidget =
 			SNew(SBox)
diff --git a/Engine/Plugins/Editor/AssetManagerEditor/Source/AssetManagerEditor/Private/ReferenceViewer/SReferenceViewerFilterBar.cpp b/Engine/Plugins/Editor/AssetManagerEditor/Source/AssetManagerEditor/Private/ReferenceViewer/SReferenceViewerFilterBar.cpp
--- a/Engine/Plugins/Editor/AssetManagerEditor/Source/AssetManagerEditor/Private/ReferenceViewer/SReferenceViewerFilterBar.cpp
+++ b/Engine/Plugins/Editor/AssetManagerEditor/Source/AssetManagerEditor/Private/ReferenceViewer/SReferenceViewerFilterBar.cpp
@@ -3,21 +3,40 @@
 #include "ReferenceViewer/SReferenceViewerFilterBar.h"
 
 #include "ReferenceViewer/ReferenceViewerSettings.h"
+#include "ReferenceViewerFilters.h"
+
+#define LOCTEXT_NAMESPACE "ReferenceViewerFilterBar"
+
+void SReferenceViewerFilterBar::Construct(const FArguments& InArgs)
+{
+	SFilterBar::Construct(InArgs);
+
+	TSharedPtr<FFilterCategory> DefaultCategory = MakeShareable( new FFilterCategory(LOCTEXT("ReferenceViewerFiltersCategory", "Other Filters"), LOCTEXT("ReferenceViewerFiltersCategoryTooltip", "Filter nodes by all filters in this category.")) );
+
+	AddFilter(MakeShareable(new FReferenceViewerFilter_ShowCheckedOut(DefaultCategory)));
+}
 
 void SReferenceViewerFilterBar::SaveSettings()
 {
-	TArray<FilterState> SaveFilters;
 	if (UReferenceViewerSettings* Settings = GetMutableDefault<UReferenceViewerSettings>())
 	{
 		// Only save & load user filters, not autofilters
 		if (!Settings->AutoUpdateFilters())
 		{
-			for (TSharedRef<SAssetFilter> CurrentAssetFilter : AssetFilters)
+			TArray<FilterState> SaveFilters;
+			for (const TSharedRef<SAssetFilter>& CurrentAssetFilter : AssetFilters)
 			{
 				FTopLevelAssetPath AssetFilterPath = CurrentAssetFilter->GetCustomClassFilterData()->GetClassPathName();
 				SaveFilters.Add(FilterState(AssetFilterPath, CurrentAssetFilter->IsEnabled()));
 			}
 			Settings->SetUserFilters(SaveFilters);
+
+			TMap<FString, bool> CustomFilters;
+			for (const TSharedRef<SFilter>& ActiveFilter : Filters)
+			{
+				CustomFilters.Add(ActiveFilter->GetFilterName(), ActiveFilter->IsEnabled());
+			}
+			Settings->SetCustomFilters(CustomFilters);
 		}
 	}
 }
@@ -39,7 +58,26 @@ void SReferenceViewerFilterBar::LoadSettings()
 					ToggleAssetTypeFilterEnabled(State.FilterPath, State.bIsEnabled);
 				}
 			}
+
+			const TMap<FString, bool>& CustomFilters = Settings->GetCustomFilters();
+			for (const TPair<FString, bool>& FilterElem : CustomFilters)
+			{
+				TSharedRef<FFilterBase<FReferenceNodeInfo&>>* FilterRefPtr = AllFrontendFilters.FindByPredicate(
+					[FilterName = FilterElem.Key](const TSharedRef<FFilterBase<FReferenceNodeInfo&>>& InFilter)
+					{
+						return InFilter->GetName().Equals(FilterName);
+					}
+				);
+
+				if (FilterRefPtr)
+				{
+					TSharedRef<FFilterBase<FReferenceNodeInfo&>> FilterRef = *FilterRefPtr;
+					AddFilterToBar(FilterRef);
+					FilterRef->SetActive(FilterElem.Value);
+				}
+			}
 		}
 	}
 }
 
+#undef LOCTEXT_NAMESPACE
diff --git a/Engine/Plugins/Editor/AssetManagerEditor/Source/AssetManagerEditor/Private/ReferenceViewer/SReferenceViewerFilterBar.h b/Engine/Plugins/Editor/AssetManagerEditor/Source/AssetManagerEditor/Private/ReferenceViewer/SReferenceViewerFilterBar.h
--- a/Engine/Plugins/Editor/AssetManagerEditor/Source/AssetManagerEditor/Private/ReferenceViewer/SReferenceViewerFilterBar.h
+++ b/Engine/Plugins/Editor/AssetManagerEditor/Source/AssetManagerEditor/Private/ReferenceViewer/SReferenceViewerFilterBar.h
@@ -11,6 +11,9 @@ class SReferenceViewerFilterBar : public SFilterBar< FReferenceNodeInfo& >
 
 public:
 
+	/** Constructs this widget with InArgs */
+	void Construct( const FArguments& InArgs );
+
 	/** Saves any settings to config that should be persistent between editor sessions */
 	void SaveSettings() override;
 
diff --git a/Engine/Plugins/Editor/AssetManagerEditor/Source/AssetManagerEditor/Public/ReferenceViewer/ReferenceViewerSettings.h b/Engine/Plugins/Editor/AssetManagerEditor/Source/AssetManagerEditor/Public/ReferenceViewer/ReferenceViewerSettings.h
--- a/Engine/Plugins/Editor/AssetManagerEditor/Source/AssetManagerEditor/Public/ReferenceViewer/ReferenceViewerSettings.h
+++ b/Engine/Plugins/Editor/AssetManagerEditor/Source/AssetManagerEditor/Public/ReferenceViewer/ReferenceViewerSettings.h
@@ -124,6 +124,9 @@ public:
 	const TArray<FilterState>& GetUserFilters() const;
 	void SetUserFilters(TArray<FilterState>& InFilters);
 
+	const TMap<FString, bool>& GetCustomFilters() const;
+	void SetCustomFilters(const TMap<FString, bool>& InCustomFilters);
+
 	bool GetFindPathEnabled() const;
 	void SetFindPathEnabled(bool bNewEnabled);
 
@@ -224,6 +227,10 @@ private:
 	UPROPERTY(config)
 	TArray<FilterState> UserFilters;
 
+	/* The list of reference viewer specific filters the user has built up */
+	UPROPERTY(config)
+	TMap<FString, bool> CustomFilters;
+
 	/* Enable the Find Path tool */
 	UPROPERTY(config)
 	bool bFindPathEnabled;
