[Reference Viewer] Text conflict from LOCTEXT macro for namespace 'ReferenceViewerFilters' and key 'ReferenceViewerFilter_ShowCheckedOut'

diff --git a/Engine/Plugins/Editor/AssetManagerEditor/Source/AssetManagerEditor/Private/ReferenceViewer/ReferenceViewerFilters.cpp b/Engine/Plugins/Editor/AssetManagerEditor/Source/AssetManagerEditor/Private/ReferenceViewer/ReferenceViewerFilters.cpp
--- a/Engine/Plugins/Editor/AssetManagerEditor/Source/AssetManagerEditor/Private/ReferenceViewer/ReferenceViewerFilters.cpp
+++ b/Engine/Plugins/Editor/AssetManagerEditor/Source/AssetManagerEditor/Private/ReferenceViewer/ReferenceViewerFilters.cpp
@@ -31,12 +31,12 @@ FString FReferenceViewerFilter_ShowCheckedOut::GetName() const
 
 FText FReferenceViewerFilter_ShowCheckedOut::GetDisplayName() const
 {
-	return LOCTEXT("ReferenceViewerFilter_ShowCheckedOut", "Checked Out");
+	return LOCTEXT("ReferenceViewerFilter_ShowCheckedOutLabel", "Checked Out");
 }
 
 FText FReferenceViewerFilter_ShowCheckedOut::GetToolTipText() const
 {
-	return LOCTEXT("ReferenceViewerFilter_ShowCheckedOut", "Allow display of Checked Out Items.");
+	return LOCTEXT("ReferenceViewerFilter_ShowCheckedOutTooltip", "Allow display of Checked Out Items.");
 }
 
 void FReferenceViewerFilter_ShowCheckedOut::ActiveStateChanged(bool bInActive)
