ContextMenu: Fixed unsafe call to ISourceControlProvider::GetState from a non game thread.

#rnx

diff --git a/Engine/Plugins/Editor/AssetManagerEditor/Source/AssetManagerEditor/Private/AssetSourceControlContextMenu.cpp b/Engine/Plugins/Editor/AssetManagerEditor/Source/AssetManagerEditor/Private/AssetSourceControlContextMenu.cpp
--- a/Engine/Plugins/Editor/AssetManagerEditor/Source/AssetManagerEditor/Private/AssetSourceControlContextMenu.cpp
+++ b/Engine/Plugins/Editor/AssetManagerEditor/Source/AssetManagerEditor/Private/AssetSourceControlContextMenu.cpp
@@ -209,9 +209,6 @@ private:
 	/** Try using perforce cache to update state */
 	void TryCacheCanExecuteVars(const TArray<FString>& InSelectedPaths, TArray<FString>* OptionalPathsWithUnknownState);
 
-	/** Cancel any currently running perforce operation */
-	void CancelCacheCanExecuteVars();
-
 private:
 
 	TArray<FAssetData> SelectedAssets; // Expanded list of assets, no folders
@@ -1240,7 +1237,7 @@ void FAssetSourceControlContextMenuState::NextAsyncState()
 	case EAsyncState::None:
 		if (TrySetAsyncState(EAsyncState::None, EAsyncState::TryCacheCanExecute))
 		{
-			Async(EAsyncExecution::ThreadPool, [this, This = AsShared()]()
+			Async(EAsyncExecution::TaskGraphMainThread, [this, This = AsShared()]()
 				{
 					TArray<FString> SelectedPaths;
 					Algo::TransformIf(SelectedAssets, SelectedPaths,
@@ -1249,8 +1246,7 @@ void FAssetSourceControlContextMenuState::NextAsyncState()
 
 					TryCacheCanExecuteVars(SelectedPaths, &PathsWithUnknownState);
 
-					// Source control operations have to be started from the main thread
-					Async(EAsyncExecution::TaskGraphMainThread, [this, This = AsShared()]() { NextAsyncState(); });
+					NextAsyncState();
 				});
 		}
 		break;
@@ -1320,7 +1316,7 @@ void FAssetSourceControlContextMenuState::TryCacheCanExecuteVars(const TArray<FS
 		}
 	}
 
-	// Check the SCC state for each package in the selected paths
+	// Check the SCC state for each package in the selected paths - this needs to happen on the main thread
 	TArray<FSourceControlStateRef> SelectedStates;
 	if (SourceControlProvider.GetState(InSelectedPaths, SelectedStates, EStateCacheUsage::Use) == ECommandResult::Succeeded)
 	{
diff --git a/Engine/Plugins/Editor/ContentBrowser/ContentBrowserFileDataSource/Source/ContentBrowserFileDataSource/Private/FileSourceControlContextMenu.cpp b/Engine/Plugins/Editor/ContentBrowser/ContentBrowserFileDataSource/Source/ContentBrowserFileDataSource/Private/FileSourceControlContextMenu.cpp
--- a/Engine/Plugins/Editor/ContentBrowser/ContentBrowserFileDataSource/Source/ContentBrowserFileDataSource/Private/FileSourceControlContextMenu.cpp
+++ b/Engine/Plugins/Editor/ContentBrowser/ContentBrowserFileDataSource/Source/ContentBrowserFileDataSource/Private/FileSourceControlContextMenu.cpp
@@ -697,12 +697,11 @@ void FFileSourceControlContextMenuState::NextAsyncState()
 	case EAsyncState::None:
 		if (TrySetAsyncState(EAsyncState::None, EAsyncState::TryCacheCanExecute))
 		{
-			Async(EAsyncExecution::ThreadPool, [this, This = AsShared()]()
+			Async(EAsyncExecution::TaskGraphMainThread, [this, This = AsShared()]()
 			{
 				TryCacheCanExecuteVars(SelectedFiles, &PathsWithUnknownState);
 
-				// Source control operations have to be started from the main thread
-				Async(EAsyncExecution::TaskGraphMainThread, [this, This = AsShared()]() { NextAsyncState(); });
+				NextAsyncState();
 			});
 		}
 		break;
@@ -752,7 +751,7 @@ void FFileSourceControlContextMenuState::TryCacheCanExecuteVars(const TArray<FSt
 	const bool bUsesCheckout = SourceControlProvider.UsesCheckout();
 	const bool bUsesFileRevisions = SourceControlProvider.UsesFileRevisions();
 
-	// Check the SCC state for each package in the selected paths
+	// Check the SCC state for each package in the selected paths - this needs to happen on the main thread
 	TArray<FSourceControlStateRef> SelectedStates;
 	if (SourceControlProvider.GetState(InSelectedPaths, SelectedStates, EStateCacheUsage::Use) == ECommandResult::Succeeded)
 	{
