Track string table key references as searchable names

Also updated the reference viewer to show "In {AssetType}: {AssetName}" for any asset type that hosts a searchable name, rather than being hardcoded for Data Tables

diff --git a/Engine/Plugins/Editor/AssetManagerEditor/Source/AssetManagerEditor/Private/ReferenceViewer/EdGraphNode_Reference.cpp b/Engine/Plugins/Editor/AssetManagerEditor/Source/AssetManagerEditor/Private/ReferenceViewer/EdGraphNode_Reference.cpp
--- a/Engine/Plugins/Editor/AssetManagerEditor/Source/AssetManagerEditor/Private/ReferenceViewer/EdGraphNode_Reference.cpp
+++ b/Engine/Plugins/Editor/AssetManagerEditor/Source/AssetManagerEditor/Private/ReferenceViewer/EdGraphNode_Reference.cpp
@@ -52,18 +52,26 @@ void UEdGraphNode_Reference::SetupReferenceNode(const FIntPoint& NodeLoc, const
 	Identifiers = NewIdentifiers;
 	const FAssetIdentifier& First = NewIdentifiers[0];
 	FString MainAssetName = InAssetData.AssetName.ToString();
-	FString AssetTypeName = InAssetData.AssetClassPath.GetAssetName().ToString();
+	FText AssetTypeDisplayName;
 
 	FAssetToolsModule& AssetToolsModule = FModuleManager::LoadModuleChecked<FAssetToolsModule>(TEXT("AssetTools"));	
 	if (UClass* AssetClass = InAssetData.GetClass())
 	{
-		TWeakPtr<IAssetTypeActions> AssetTypeActions = AssetToolsModule.Get().GetAssetTypeActionsForClass(InAssetData.GetClass());
-		if(AssetTypeActions.IsValid())
+		AssetTypeDisplayName = AssetClass->GetDisplayNameText();
+
+		if (TSharedPtr<IAssetTypeActions> AssetTypeActions = AssetToolsModule.Get().GetAssetTypeActionsForClass(AssetClass).Pin())
 		{
-			AssetTypeColor = AssetTypeActions.Pin()->GetTypeColor();
+			AssetTypeColor = AssetTypeActions->GetTypeColor();
 		}
 	}
-	AssetBrush = FSlateIcon("EditorStyle", FName( *("ClassIcon." + AssetTypeName)));
+	else if(InAssetData.IsValid())
+	{
+		AssetTypeDisplayName = FText::AsCultureInvariant(InAssetData.AssetClassPath.GetAssetName().ToString());
+	}
+	if (InAssetData.IsValid())
+	{
+		AssetBrush = FSlateIcon("EditorStyle", FName(*("ClassIcon." + InAssetData.AssetClassPath.GetAssetName().ToString())));
+	}
 
 	bIsCollapsed = false;
 	bIsPackage = true;
@@ -73,21 +81,35 @@ void UEdGraphNode_Reference::SetupReferenceNode(const FIntPoint& NodeLoc, const
 	FPrimaryAssetId PrimaryAssetID = NewIdentifiers[0].GetPrimaryAssetId();
 	if (PrimaryAssetID.IsValid())  // Management References (PrimaryAssetIDs)
 	{
-		static FText ManagerText = LOCTEXT("ReferenceManager", "Manager");
+		static const FText ManagerText = LOCTEXT("ReferenceManager", "Manager");
 		MainAssetName = PrimaryAssetID.PrimaryAssetType.ToString() + TEXT(":") + PrimaryAssetID.PrimaryAssetName.ToString();
-		AssetTypeName = ManagerText.ToString();
+		AssetTypeDisplayName = ManagerText;
 		bIsPackage = false;
 		bIsPrimaryAsset = true;
 	}
 	else if (First.IsValue()) // Searchable Names (GamePlay Tags, Data Table Row Handle)
 	{
 		MainAssetName = First.ValueName.ToString();
-		AssetTypeName = First.ObjectName.ToString();
-		static const FName NAME_DataTable(TEXT("DataTable"));
-		static const FText InDataTableText = LOCTEXT("InDataTable", "In DataTable");
-		if (InAssetData.AssetClassPath.GetAssetName() == NAME_DataTable)
+		if (AssetTypeDisplayName.IsEmpty())
+		{
+			static const FName GameplayTagTypeName = "GameplayTag";
+			if (First.ObjectName == GameplayTagTypeName)
+			{
+				static const FText GameplayTagText = LOCTEXT("ReferenceGameplayTag", "Gameplay Tag");
+				AssetTypeDisplayName = GameplayTagText;
+			}
+			else
+			{
+				AssetTypeDisplayName = FText::AsCultureInvariant(First.ObjectName.ToString());
+			}
+		}
+		else
 		{
-			AssetTypeName = InDataTableText.ToString() + TEXT(" ") + AssetTypeName;
+			AssetTypeDisplayName = FText::FormatNamed(
+				LOCTEXT("InAssetFmt", "In {AssetType}: {AssetName}"), 
+				TEXT("AssetType"), AssetTypeDisplayName,
+				TEXT("AssetName"), FText::AsCultureInvariant(First.ObjectName.ToString())
+				);
 		}
 
 		bIsPackage = false;
@@ -97,8 +119,9 @@ void UEdGraphNode_Reference::SetupReferenceNode(const FIntPoint& NodeLoc, const
 		const FString PackageNameStr = Identifiers[0].PackageName.ToString();
 		if ( PackageNameStr.StartsWith(TEXT("/Script")) )// C++ Packages (/Script Code)
 		{
+			static const FText ScriptText = LOCTEXT("ReferenceScript", "Script");
 			MainAssetName = PackageNameStr.RightChop(8);
-			AssetTypeName = TEXT("Script");
+			AssetTypeDisplayName = ScriptText;
 		}
 	}
 
@@ -108,11 +131,13 @@ void UEdGraphNode_Reference::SetupReferenceNode(const FIntPoint& NodeLoc, const
 		InAssetData.GetTagValue(NAME_ActorLabel, MainAssetName); 
 
 		// append the type so it shows up on the extra line
-		NodeTitle = FText::FromString(FString::Printf(TEXT("%s\n%s"), *MainAssetName, *AssetTypeName));
+		static const FTextFormat NodeTitleFmt = INVTEXT("{0}\n{1}");
+		NodeTitle = FText::Format(NodeTitleFmt, FText::FromString(MainAssetName), AssetTypeDisplayName);
 	}
 	else
 	{
-		NodeTitle = FText::Format(LOCTEXT("ReferenceNodeMultiplePackagesComment", "{0} and {1} others"), FText::FromString(MainAssetName), FText::AsNumber(NewIdentifiers.Num() - 1));
+		static const FTextFormat NodeTitleFmt = LOCTEXT("ReferenceNodeMultiplePackagesComment", "{0} and {1} {1}|plural(one=other,other=others)");
+		NodeTitle = FText::Format(NodeTitleFmt, FText::FromString(MainAssetName), NewIdentifiers.Num() - 1);
 	}
 	
 	CacheAssetData(InAssetData);
diff --git a/Engine/Source/Runtime/Core/Private/Internationalization/TextHistory.cpp b/Engine/Source/Runtime/Core/Private/Internationalization/TextHistory.cpp
--- a/Engine/Source/Runtime/Core/Private/Internationalization/TextHistory.cpp
+++ b/Engine/Source/Runtime/Core/Private/Internationalization/TextHistory.cpp
@@ -2605,7 +2605,9 @@ FTextId FTextHistory_StringTableEntry::FStringTableReferenceData::GetTextId()
 
 void FTextHistory_StringTableEntry::FStringTableReferenceData::CollectStringTableAssetReferences(FStructuredArchive::FRecord Record)
 {
-	if (Record.GetUnderlyingArchive().IsObjectReferenceCollector())
+	FArchive& BaseArchive = Record.GetUnderlyingArchive();
+
+	if (BaseArchive.IsObjectReferenceCollector())
 	{
 		UE::TScopeLock ScopeLock(DataCS);
 
@@ -2619,6 +2621,8 @@ void FTextHistory_StringTableEntry::FStringTableReferenceData::CollectStringTabl
 			DisplayString.Reset();
 		}
 	}
+
+	IStringTableEngineBridge::CollectStringTableAssetSearchableName(TableId, Key, BaseArchive);
 }
 
 FStringTableEntryConstPtr FTextHistory_StringTableEntry::FStringTableReferenceData::ResolveStringTableEntry()
diff --git a/Engine/Source/Runtime/Core/Public/Internationalization/StringTableCore.h b/Engine/Source/Runtime/Core/Public/Internationalization/StringTableCore.h
--- a/Engine/Source/Runtime/Core/Public/Internationalization/StringTableCore.h
+++ b/Engine/Source/Runtime/Core/Public/Internationalization/StringTableCore.h
@@ -289,6 +289,15 @@ public:
 		}
 	}
 
+	/** Collect any searchable name data for a string table asset entry reference */
+	static void CollectStringTableAssetSearchableName(const FName InTableId, const FTextKey InKey, const FArchive& Ar)
+	{
+		if (InstancePtr)
+		{
+			InstancePtr->CollectStringTableAssetSearchableNameImpl(InTableId, InKey, Ar);
+		}
+	}
+
 	/** Is this string table from an asset? */
 	static bool IsStringTableFromAsset(const FName InTableId)
 	{
@@ -309,6 +318,7 @@ protected:
 	virtual void FullyLoadStringTableAssetImpl(FName& InOutTableId) = 0;
 	virtual void RedirectStringTableAssetImpl(FName& InOutTableId) = 0;
 	virtual void CollectStringTableAssetReferencesImpl(FName& InOutTableId, FStructuredArchive::FSlot Slot) = 0;
+	virtual void CollectStringTableAssetSearchableNameImpl(const FName InTableId, const FTextKey InKey, const FArchive& Ar) = 0;
 	virtual bool IsStringTableFromAssetImpl(const FName InTableId) = 0;
 	virtual bool IsStringTableAssetBeingReplacedImpl(const UStringTable* InStringTableAsset) = 0;
 
diff --git a/Engine/Source/Runtime/Engine/Private/Internationalization/StringTable.cpp b/Engine/Source/Runtime/Engine/Private/Internationalization/StringTable.cpp
--- a/Engine/Source/Runtime/Engine/Private/Internationalization/StringTable.cpp
+++ b/Engine/Source/Runtime/Engine/Private/Internationalization/StringTable.cpp
@@ -288,6 +288,20 @@ private:
 		}
 	}
 
+	virtual void CollectStringTableAssetSearchableNameImpl(const FName InTableId, const FTextKey InKey, const FArchive& Ar) override
+	{
+		if (Ar.IsSaving() && !InTableId.IsNone() && !InKey.IsEmpty())
+		{
+			if (UStringTable* StringTableAsset = FStringTableRegistry::Get().FindStringTableAsset(InTableId))
+			{
+				// Note which key we are pointing to for later searching
+				FNameBuilder KeyStr;
+				InKey.ToString(KeyStr);
+				Ar.MarkSearchableName(StringTableAsset, FName(KeyStr));
+			}
+		}
+	}
+
 	virtual bool IsStringTableFromAssetImpl(const FName InTableId) override
 	{
 		const FSoftObjectPath StringTableAssetReference = GetAssetReference(InTableId);
