FAssetManagerEditorModule : Load the SourceControl module from StartupModule() to ensure that the SourceControl module outlives this module's usage of it.

diff --git a/Engine/Plugins/Editor/AssetManagerEditor/Source/AssetManagerEditor/Private/AssetManagerEditorModule.cpp b/Engine/Plugins/Editor/AssetManagerEditor/Source/AssetManagerEditor/Private/AssetManagerEditorModule.cpp
--- a/Engine/Plugins/Editor/AssetManagerEditor/Source/AssetManagerEditor/Private/AssetManagerEditorModule.cpp
+++ b/Engine/Plugins/Editor/AssetManagerEditor/Source/AssetManagerEditor/Private/AssetManagerEditorModule.cpp
@@ -69,6 +69,7 @@
 #include "ContentBrowserDataSubsystem.h"
 #include "TreeView/AssetTable.h"
 #include "TreeView/SAssetTableTreeView.h"
+#include "ISourceControlModule.h"
 
 #define LOCTEXT_NAMESPACE "AssetManagerEditor"
 
@@ -487,6 +488,9 @@ void FAssetManagerEditorModule::StartupModule()
 	EditorCookedSandbox = nullptr;
 	CurrentRegistrySource = nullptr;
 
+	// Load the source control module to ensure that it outlives our use of it.
+	FModuleManager::Get().LoadModuleChecked<ISourceControlModule>(TEXT("SourceControl"));
+
 	// Load the tree map module so we can use it for size map
 	FModuleManager::Get().LoadModule(TEXT("TreeMap"));
 
