Expose AssetTableTreeView as part of context menu in PluginReferenceViewer to make it easier to find.

#Synced-CL 50308174

diff --git a/Engine/Plugins/Editor/AssetManagerEditor/Source/AssetManagerEditor/Private/AssetManagerEditorModule.cpp b/Engine/Plugins/Editor/AssetManagerEditor/Source/AssetManagerEditor/Private/AssetManagerEditorModule.cpp
--- a/Engine/Plugins/Editor/AssetManagerEditor/Source/AssetManagerEditor/Private/AssetManagerEditorModule.cpp
+++ b/Engine/Plugins/Editor/AssetManagerEditor/Source/AssetManagerEditor/Private/AssetManagerEditorModule.cpp
@@ -337,6 +337,7 @@ public:
 	virtual void OpenSizeMapUI(TArray<FAssetIdentifier> SelectedIdentifiers) override;
 	virtual void OpenSizeMapUI(TArray<FName> SelectedPackages) override;	
 	virtual void OpenShaderCookStatistics(TArray<FName> SelectedIdentifiers) override;
+	virtual void OpenAssetSizeTreeViewUI(const FString& SelectedPlugin) override;
 	virtual bool GetStringValueForCustomColumn(const FAssetData& AssetData, FName ColumnName, FString& OutValue, const FAssetManagerEditorRegistrySource* OverrideRegistrySource = nullptr) override;
 	virtual bool GetDisplayTextForCustomColumn(const FAssetData& AssetData, FName ColumnName, FText& OutValue, const FAssetManagerEditorRegistrySource* OverrideRegistrySource = nullptr) override;
 	virtual bool GetIntegerValueForCustomColumn(const FAssetData& AssetData, FName ColumnName, int64& OutValue, const FAssetManagerEditorRegistrySource* OverrideRegistrySource = nullptr) override;
@@ -1072,6 +1073,16 @@ void FAssetManagerEditorModule::OpenShaderCookStatistics(TArray<FName> SelectedP
 	}
 }
 
+void FAssetManagerEditorModule::OpenAssetSizeTreeViewUI(const FString& SearchTerm)
+{
+	if (TSharedPtr<SDockTab> NewTab = FGlobalTabmanager::Get()->TryInvokeTab(AssetDiskSizeTabName))
+	{
+		TSharedRef<SAssetTableTreeView> AssetTableTreeView = StaticCastSharedRef<SAssetTableTreeView>(NewTab->GetContent());
+
+		AssetTableTreeView->SetFilterToName(SearchTerm);
+	}
+}
+
 void FAssetManagerEditorModule::GetAssetDataInPaths(const TArray<FString>& Paths, TArray<FAssetData>& OutAssetData)
 {
 	// Form a filter from the paths
diff --git a/Engine/Plugins/Editor/AssetManagerEditor/Source/AssetManagerEditor/Private/TreeView/SAssetTableTreeView.cpp b/Engine/Plugins/Editor/AssetManagerEditor/Source/AssetManagerEditor/Private/TreeView/SAssetTableTreeView.cpp
--- a/Engine/Plugins/Editor/AssetManagerEditor/Source/AssetManagerEditor/Private/TreeView/SAssetTableTreeView.cpp
+++ b/Engine/Plugins/Editor/AssetManagerEditor/Source/AssetManagerEditor/Private/TreeView/SAssetTableTreeView.cpp
@@ -14,16 +14,11 @@
 #include "Framework/MultiBox/MultiBoxBuilder.h"
 #include "GenericPlatform/GenericPlatformMisc.h"
 #include "HAL/FileManager.h"
-#include "HAL/PlatformFileManager.h"
 #include "IContentBrowserSingleton.h"
 #include "InsightsCore/Common/Log.h"
 #include "InsightsCore/Common/Stopwatch.h"
-#include "InsightsCore/Filter/ViewModels/FilterConfigurator.h"
-#include "InsightsCore/Table/ViewModels/TableCellValueFormatter.h"
-#include "InsightsCore/Table/ViewModels/TableColumn.h"
 #include "InsightsCore/Table/ViewModels/TreeNodeGrouping.h"
 #include "Interfaces/IPluginManager.h"
-#include "Logging/MessageLog.h"
 #include "Misc/FileHelper.h"
 #include "Modules/ModuleManager.h"
 #include "Serialization/ArrayReader.h"
@@ -36,11 +31,7 @@
 #include "TreeView/AssetDependencyGrouping.h"
 #include "Widgets/Images/SImage.h"
 #include "Widgets/Input/SButton.h"
-#include "Widgets/Input/SCheckBox.h"
-#include "Widgets/SToolTip.h"
-
-#include <limits>
-#include <memory>
+#include "Widgets/Input/SSearchBox.h"
 
 #define LOCTEXT_NAMESPACE "SAssetTableTreeView"
 
@@ -1652,6 +1643,18 @@ void SAssetTableTreeView::TreeView_OnSelectionChanged(UE::Insights::FTableTreeNo
 	}
 }
 
+TSharedRef<SWidget> SAssetTableTreeView::ConstructSearchBox()
+{
+	SAssignNew(SearchBox, SSearchBox)
+		.HintText(LOCTEXT("SearchBox_Hint", "Search"))
+		.ToolTipText(LOCTEXT("SearchBox_ToolTip", "Type here to search the tree hierarchy by item or group name."))
+		.OnTextChanged(this, &SAssetTableTreeView::SearchBox_OnTextChanged);
+
+	SearchBoxWidget = SearchBox;
+
+	return SearchBoxWidget.ToSharedRef();
+}
+
 ////////////////////////////////////////////////////////////////////////////////////////////////////
 
 void SAssetTableTreeView::QuickPopulateAssetTableRow(FAssetTableRow& OutRow, const FAssetData& AssetData, FAssetTable& AssetTable) const
@@ -2341,6 +2344,14 @@ void SAssetTableTreeView::RebuildTree(bool bResync)
 		TotalTime, SyncTime, TotalTime - SyncTime, TableRowNodes.Num());
 }
 
+void SAssetTableTreeView::SetFilterToName(const FString& Name)
+{
+	if (SearchBoxWidget.IsValid())
+	{
+		StaticCastSharedPtr<SSearchBox>(SearchBoxWidget)->SetText(FText::FromString(Name));
+	}
+}
+
 ////////////////////////////////////////////////////////////////////////////////////////////////////
 
 #undef LOCTEXT_NAMESPACE
diff --git a/Engine/Plugins/Editor/AssetManagerEditor/Source/AssetManagerEditor/Private/TreeView/SAssetTableTreeView.h b/Engine/Plugins/Editor/AssetManagerEditor/Source/AssetManagerEditor/Private/TreeView/SAssetTableTreeView.h
--- a/Engine/Plugins/Editor/AssetManagerEditor/Source/AssetManagerEditor/Private/TreeView/SAssetTableTreeView.h
+++ b/Engine/Plugins/Editor/AssetManagerEditor/Source/AssetManagerEditor/Private/TreeView/SAssetTableTreeView.h
@@ -62,6 +62,12 @@ public:
 	 */
 	virtual void RebuildTree(bool bResync);
 
+	/**
+	 * Sets the search filter to the given name (does not check if exists but UI is clear when no data matches filter)
+	 * @param Name - Could be an asset name or a plugin name (any name in the tree really)
+	 */
+	void SetFilterToName(const FString& Name);
+
 	////////////////////////////////////////////////////////////////////////////////////////////////////
 	// IAsyncOperationStatusProvider implementation
 
@@ -106,6 +112,8 @@ private:
 
 	virtual void TreeView_OnSelectionChanged(UE::Insights::FTableTreeNodePtr SelectedItem, ESelectInfo::Type SelectInfo) override;
 
+	virtual TSharedRef<SWidget> ConstructSearchBox() override;
+
 	void InitAvailableViewPresets();
 
 	void ExportDependencyData() const;
@@ -155,6 +163,9 @@ private:
 	/** List of valid registry sources */
 	TArray<TSharedPtr<FString>> SourceComboList;
 
+	/** Used so we can set the search filter from code and have it reflected in the UI */
+	TSharedPtr<SWidget> SearchBoxWidget;
+
 	/**
 	* Text block where we put the timestamp for the asset registry so the user knows
 	* if they are looking at super out of date data.
diff --git a/Engine/Plugins/Editor/AssetManagerEditor/Source/AssetManagerEditor/Public/AssetManagerEditorModule.h b/Engine/Plugins/Editor/AssetManagerEditor/Source/AssetManagerEditor/Public/AssetManagerEditorModule.h
--- a/Engine/Plugins/Editor/AssetManagerEditor/Source/AssetManagerEditor/Public/AssetManagerEditorModule.h
+++ b/Engine/Plugins/Editor/AssetManagerEditor/Source/AssetManagerEditor/Public/AssetManagerEditorModule.h
@@ -443,6 +443,9 @@ public:
 	/** Open the Shader cook stats */
 	virtual void OpenShaderCookStatistics(TArray<FName> SelectedIdentifiers) = 0;
 
+	/** Open the Asset Disk Size Tree View*/
+	virtual void OpenAssetSizeTreeViewUI(const FString& SearchTerm) = 0;
+
 	/** Called by SReferenceViewer to update given list of commands with any of the registered delegates */
 	virtual void AppendRegisteredCommandsToCommandList(TSharedRef<FUICommandList> ListToAppendTo) = 0;
 
diff --git a/Engine/Plugins/Experimental/PluginReferenceViewer/Source/PluginReferenceViewer/Private/PluginReferenceViewerCommands.cpp b/Engine/Plugins/Experimental/PluginReferenceViewer/Source/PluginReferenceViewer/Private/PluginReferenceViewerCommands.cpp
--- a/Engine/Plugins/Experimental/PluginReferenceViewer/Source/PluginReferenceViewer/Private/PluginReferenceViewerCommands.cpp
+++ b/Engine/Plugins/Experimental/PluginReferenceViewer/Source/PluginReferenceViewer/Private/PluginReferenceViewerCommands.cpp
@@ -19,8 +19,8 @@ void FPluginReferenceViewerCommands::RegisterCommands()
 	UI_COMMAND(ShowDuplicates, "Show Duplicates References", "Toggles visibility of duplicates references", EUserInterfaceActionType::ToggleButton, FInputChord(EKeys::U));
 	UI_COMMAND(ShowEnginePlugins, "Show Engine Plugins", "Toggles visibility of Engine Plugins", EUserInterfaceActionType::ToggleButton, FInputChord());
 	UI_COMMAND(ShowOptionalPlugins, "Show Optional Plugins", "Toggles visibility of Optional Plugin dependencies", EUserInterfaceActionType::ToggleButton, FInputChord());
-
 	UI_COMMAND(OpenPluginProperties, "Open Plugin Properties", "Opens the Plugin Properties Editor Window", EUserInterfaceActionType::Button, FInputChord());
+	UI_COMMAND(OpenAssetDiskSizeViewer, "Open Asset Disk Size Viewer", "Opens the Asset Disk Size Window", EUserInterfaceActionType::Button, FInputChord());
 	UI_COMMAND(ShowInContentBrowser, "Browse to Plugin", "Browses to the associated Plugin and selects it in the most recently used Content Browser (summoning one if necessary)", EUserInterfaceActionType::Button, FInputChord(EModifierKey::Control, EKeys::B));
 	UI_COMMAND(ZoomToFit, "Zoom to Fit", "Zoom in and center the view on the selected item", EUserInterfaceActionType::Button, FInputChord(EKeys::F));
 	UI_COMMAND(ReCenterGraph, "Re-Center Graph", "Re-centers the graph on this node, showing all referencers and references for this Plugin instead", EUserInterfaceActionType::Button, FInputChord());
diff --git a/Engine/Plugins/Experimental/PluginReferenceViewer/Source/PluginReferenceViewer/Private/PluginReferenceViewerCommands.h b/Engine/Plugins/Experimental/PluginReferenceViewer/Source/PluginReferenceViewer/Private/PluginReferenceViewerCommands.h
--- a/Engine/Plugins/Experimental/PluginReferenceViewer/Source/PluginReferenceViewer/Private/PluginReferenceViewerCommands.h
+++ b/Engine/Plugins/Experimental/PluginReferenceViewer/Source/PluginReferenceViewer/Private/PluginReferenceViewerCommands.h
@@ -20,6 +20,7 @@ public:
 	TSharedPtr<FUICommandInfo> ShowOptionalPlugins;
 
 	TSharedPtr<FUICommandInfo> OpenPluginProperties;
+	TSharedPtr<FUICommandInfo> OpenAssetDiskSizeViewer;
 	TSharedPtr<FUICommandInfo> ShowInContentBrowser;
 	TSharedPtr<FUICommandInfo> ZoomToFit;
 	TSharedPtr<FUICommandInfo> ReCenterGraph;
diff --git a/Engine/Plugins/Experimental/PluginReferenceViewer/Source/PluginReferenceViewer/Private/PluginReferenceViewerSchema.cpp b/Engine/Plugins/Experimental/PluginReferenceViewer/Source/PluginReferenceViewer/Private/PluginReferenceViewerSchema.cpp
--- a/Engine/Plugins/Experimental/PluginReferenceViewer/Source/PluginReferenceViewer/Private/PluginReferenceViewerSchema.cpp
+++ b/Engine/Plugins/Experimental/PluginReferenceViewer/Source/PluginReferenceViewer/Private/PluginReferenceViewerSchema.cpp
@@ -63,6 +63,8 @@ void UPluginReferenceViewerSchema::GetContextMenuActions(UToolMenu* Menu, UGraph
 		FToolMenuSection& Section = Menu->AddSection(TEXT("Plugin"), NSLOCTEXT("PluginReferenceViewerSchema", "PluginSectionLabel", "Plugin"));
 		
 		Section.AddMenuEntry(FPluginReferenceViewerCommands::Get().OpenPluginProperties);
+
+		Section.AddMenuEntry(FPluginReferenceViewerCommands::Get().OpenAssetDiskSizeViewer);
 		
 		FToolMenuEntry& ShowInEntry = Section.AddMenuEntry(FPluginReferenceViewerCommands::Get().ShowInContentBrowser);
 		ShowInEntry.Icon = FSlateIcon(FAppStyle::GetAppStyleSetName(), "SystemWideCommands.FindInContentBrowser.Small");
diff --git a/Engine/Plugins/Experimental/PluginReferenceViewer/Source/PluginReferenceViewer/Private/SPluginReferenceViewer.cpp b/Engine/Plugins/Experimental/PluginReferenceViewer/Source/PluginReferenceViewer/Private/SPluginReferenceViewer.cpp
--- a/Engine/Plugins/Experimental/PluginReferenceViewer/Source/PluginReferenceViewer/Private/SPluginReferenceViewer.cpp
+++ b/Engine/Plugins/Experimental/PluginReferenceViewer/Source/PluginReferenceViewer/Private/SPluginReferenceViewer.cpp
@@ -2,6 +2,7 @@
 
 #include "SPluginReferenceViewer.h"
 
+#include "AssetManagerEditorModule.h"
 #include "ContentBrowserModule.h"
 #include "EdGraphNode_PluginReference.h"
 #include "EdGraph_PluginReferenceViewer.h"
@@ -186,6 +187,20 @@ void SPluginReferenceViewer::OnOpenPluginProperties()
 	}
 }
 
+void SPluginReferenceViewer::OnOpenAssetDiskSizeViewer()
+{
+	if (UEdGraphNode* SelectedNode = GraphEditorPtr->GetSingleSelectedNode())
+	{
+		if (UEdGraphNode_PluginReference* PluginReferenceNode = Cast<UEdGraphNode_PluginReference>(SelectedNode))
+		{
+			if (TSharedPtr<const IPlugin> PluginReferenced = PluginReferenceNode->GetPlugin())
+			{
+				IAssetManagerEditorModule::Get().OpenAssetSizeTreeViewUI(PluginReferenced->GetName());
+			}
+		}
+	}
+}
+
 void SPluginReferenceViewer::OnShowPluginInContentBrowser()
 {
 	TSharedPtr<const IPlugin> Plugin;
@@ -266,7 +281,21 @@ void SPluginReferenceViewer::OnUpdateHistoryData(FPluginReferenceViewerHistoryDa
 
 bool SPluginReferenceViewer::HasAtLeastOneRealNodeSelected()
 {
-	return true;
+	// NOTE: Based on SReferenceViewer::HasAtLeastOneRealNodeSelected without need to check for collapsed node
+	if (GraphEditorPtr.IsValid())
+	{
+		TSet<UObject*> SelectedNodes = GraphEditorPtr->GetSelectedNodes();
+		for (UObject* Node : SelectedNodes)
+		{
+			UEdGraphNode_PluginReference* ReferenceNode = Cast<UEdGraphNode_PluginReference>(Node);
+			if (ReferenceNode)
+			{
+				return true;
+			}
+		}
+	}
+
+	return false;
 }
 
 void SPluginReferenceViewer::OpenPluginProperties(const FString& PluginName)
@@ -557,6 +586,11 @@ void SPluginReferenceViewer::RegisterActions()
 		FPluginReferenceViewerCommands::Get().OpenPluginProperties,
 		FExecuteAction::CreateSP(this, &SPluginReferenceViewer::OnOpenPluginProperties),
 		FCanExecuteAction::CreateSP(this, &SPluginReferenceViewer::HasAtLeastOneRealNodeSelected));
+
+	PluginReferenceViewerActions->MapAction(
+		FPluginReferenceViewerCommands::Get().OpenAssetDiskSizeViewer,
+		FExecuteAction::CreateSP(this, &SPluginReferenceViewer::OnOpenAssetDiskSizeViewer),
+		FCanExecuteAction::CreateSP(this, &SPluginReferenceViewer::HasAtLeastOneRealNodeSelected));
 	
 	PluginReferenceViewerActions->MapAction(
 		FPluginReferenceViewerCommands::Get().ShowInContentBrowser,
diff --git a/Engine/Plugins/Experimental/PluginReferenceViewer/Source/PluginReferenceViewer/Private/SPluginReferenceViewer.h b/Engine/Plugins/Experimental/PluginReferenceViewer/Source/PluginReferenceViewer/Private/SPluginReferenceViewer.h
--- a/Engine/Plugins/Experimental/PluginReferenceViewer/Source/PluginReferenceViewer/Private/SPluginReferenceViewer.h
+++ b/Engine/Plugins/Experimental/PluginReferenceViewer/Source/PluginReferenceViewer/Private/SPluginReferenceViewer.h
@@ -45,6 +45,7 @@ public:
 
 private:
 	void OnOpenPluginProperties();
+	void OnOpenAssetDiskSizeViewer();
 	void OnShowPluginInContentBrowser();
 	bool HasAtLeastOneRealNodeSelected();
 
