Asset Size Map - Added support for displaying external packages (editor) and generated packages (cook)
* For editor data, assets now show their external actor/object packages grouped by class as drillable children in the tree map.
* For cooked data, assets show generated packages in a "Generated Packages" container.

#Synced-CL 50811210

diff --git a/Engine/Plugins/Editor/AssetManagerEditor/Source/AssetManagerEditor/Private/SSizeMap.cpp b/Engine/Plugins/Editor/AssetManagerEditor/Source/AssetManagerEditor/Private/SSizeMap.cpp
--- a/Engine/Plugins/Editor/AssetManagerEditor/Source/AssetManagerEditor/Private/SSizeMap.cpp
+++ b/Engine/Plugins/Editor/AssetManagerEditor/Source/AssetManagerEditor/Private/SSizeMap.cpp
@@ -24,6 +24,8 @@
 #include "STreeMap.h"
 #include "Subsystems/AssetEditorSubsystem.h"
 #include "Widgets/Input/SComboBox.h"
+#include "Engine/ExternalObjectAndActorDependencyGatherer.h"
+#include "HAL/PlatformApplicationMisc.h"
 
 #include UE_INLINE_GENERATED_CPP_BY_NAME(SSizeMap)
 
@@ -195,6 +197,11 @@ void SSizeMap::Construct(const FArguments& InArgs)
 		FExecuteAction::CreateSP(this, &SSizeMap::MakeCollectionWithDependencies, ECollectionShareType::CST_Shared),
 		FCanExecuteAction::CreateSP(this, &SSizeMap::CanMakeCollectionWithDependencies, ECollectionShareType::CST_Shared));
 
+	Commands->MapAction(
+		FAssetManagerEditorCommands::Get().CopyPaths,
+		FExecuteAction::CreateSP(this, &SSizeMap::CopyPackageName),
+		FCanExecuteAction::CreateSP(this, &SSizeMap::IsAnythingSelected));
+
 	CurrentRegistrySource = EditorModule->GetCurrentRegistrySource(false);
 	bMemorySizeCached = false;
 
@@ -583,9 +590,7 @@ void SSizeMap::GatherDependenciesRecursively(TSharedPtr<FAssetThumbnailPool>& In
 				NodeSizeMapData.AssetData.AssetName = AssetPackageName;
 				NodeSizeMapData.AssetData.AssetClassPath = FTopLevelAssetPath(TEXT("/None"), *LOCTEXT("MissingAsset", "MISSING!").ToString());
 
-				const FString AssetPathString = AssetPackageNameString + TEXT(".") + FPackageName::GetLongPackageAssetName(AssetPackageNameString);
-				FAssetData FoundData = CurrentRegistrySource->GetAssetByObjectPath(FSoftObjectPath(AssetPathString));
-
+				FAssetData FoundData = CurrentRegistrySource->GetAssetByPackageName(AssetPackageName);
 				if (FoundData.IsValid())
 				{
 					NodeSizeMapData.AssetData = MoveTemp(FoundData);
@@ -598,6 +603,7 @@ void SSizeMap::GatherDependenciesRecursively(TSharedPtr<FAssetThumbnailPool>& In
 			
 			NodeSizeMapData.AssetSize = 0;
 			NodeSizeMapData.bHasKnownSize = false;
+			NodeSizeMapData.ExternalPackageCount = 0;
 
 			if (NodeSizeMapData.AssetData.IsValid())
 			{
@@ -696,6 +702,9 @@ void SSizeMap::GatherDependenciesRecursively(TSharedPtr<FAssetThumbnailPool>& In
 					NodeSizeMapData.bHasKnownSize = true;
 				}
 
+				// External / generated packages discovery
+				GatherExternalPackages(NodeSizeMapData.AssetData, AssetPackageName, ChildTreeMapNode, VisitedAssetIdentifiers);
+
 				// Now visit all of the assets that we are referencing
 				GatherDependenciesRecursively(InAssetThumbnailPool, VisitedAssetIdentifiers, ReferencedAssetIdentifiers, ChunkId != INDEX_NONE ? AssetPrimaryId : FilterPrimaryAsset, ChildTreeMapNode, SharedRootNode, NumAssetsWhichFailedToLoad);
 			}
@@ -707,6 +716,131 @@ void SSizeMap::GatherDependenciesRecursively(TSharedPtr<FAssetThumbnailPool>& In
 	}
 }
 
+void SSizeMap::GatherExternalPackages(const FAssetData CurrentAssetData, FName AssetPackageName, TSharedRef<FTreeMapNodeData> ChildTreeMapNode, TMap<FAssetIdentifier, TSharedPtr<FTreeMapNodeData>>& VisitedAssetIdentifiers)
+{
+	auto GetExternalPackageSize = [this](const FAssetData& ExternalPackageAsset) -> int64
+	{
+		int64 Size = 0;
+		if (EditorModule->GetIntegerValueForCustomColumn(ExternalPackageAsset, GetDefault<USizeMapSettings>()->SizeType, Size, CurrentRegistrySource))
+		{
+			return Size;
+		}
+		return 0;
+	};
+
+	auto AddExternalPackageNode = [&](const FAssetData& ExternalPackageAsset, const TSharedPtr<FTreeMapNodeData>& ParentNode)
+	{
+		FAssetIdentifier AssetIdentifier(ExternalPackageAsset.PackageName);
+		if (VisitedAssetIdentifiers.Contains(AssetIdentifier))
+		{
+			return;
+		}
+
+		FTreeMapNodeDataRef TreeNode = MakeShareable(new FTreeMapNodeData());
+		ParentNode->Children.Add(TreeNode);
+		TreeNode->Parent = ParentNode.Get();
+
+		VisitedAssetIdentifiers.Add(AssetIdentifier, TreeNode);
+
+		FNodeSizeMapData& ExternalPackageNode = NodeSizeMapDataMap.Add(TreeNode);
+		ExternalPackageNode.AssetData = ExternalPackageAsset;
+		ExternalPackageNode.AssetSize = GetExternalPackageSize(ExternalPackageAsset);
+		ExternalPackageNode.bHasKnownSize = true;
+	};
+
+	if (CurrentRegistrySource->bIsEditor)
+	{
+		// Scan external actor and object paths
+		TArray<FAssetData> ExternalPackagesAssets;
+		{
+			FARFilter QueryFilter = FExternalObjectAndActorDependencyGatherer::GetQueryFilter(AssetPackageName);
+			CurrentRegistrySource->GetAssets(QueryFilter, ExternalPackagesAssets);
+		}
+
+		// Group external packages by class
+		if (ExternalPackagesAssets.Num() > 0)
+		{
+			TMap<FTopLevelAssetPath, TArray<const FAssetData*>> ClassGroups;
+			for (const FAssetData& ExternalPackageAsset : ExternalPackagesAssets)
+			{
+				ClassGroups.FindOrAdd(ExternalPackageAsset.AssetClassPath).Add(&ExternalPackageAsset);
+			}
+
+			for (const TPair<FTopLevelAssetPath, TArray<const FAssetData*>>& ClassGroup : ClassGroups)
+			{
+				const FTopLevelAssetPath ClassPath = ClassGroup.Key;
+				const TArray<const FAssetData*> Assets = ClassGroup.Value;
+
+				int64 GroupTotalSize = 0;
+				for (const FAssetData* Asset : Assets)
+				{
+					GroupTotalSize += GetExternalPackageSize(*Asset);
+				}
+
+				FTreeMapNodeDataRef GroupNode = MakeShareable(new FTreeMapNodeData());
+				ChildTreeMapNode->Children.Add(GroupNode);
+				GroupNode->Parent = &ChildTreeMapNode.Get();
+
+				const FString ClassName = ClassPath.GetAssetName().ToString();
+				GroupNode->Name = FString::Printf(TEXT("%s  (%d %s, %s)"), *ClassName, Assets.Num(), Assets.Num() == 1 ? TEXT("package") : TEXT("packages"), *SizeMapInternals::MakeBestSizeString(GroupTotalSize, true));
+				GroupNode->Size = 0.0f;
+
+				for (const FAssetData* Asset : Assets)
+				{
+					AddExternalPackageNode(*Asset, GroupNode);
+				}
+			}
+
+			// Store count for display in FinalizeNodesRecursively
+			NodeSizeMapDataMap.FindChecked(ChildTreeMapNode).ExternalPackageCount = ExternalPackagesAssets.Num();
+		}
+	}
+	else
+	{
+		// Cooked: scan generated packages
+		TArray<FAssetData> GeneratedAssets;
+		{
+			const FString GeneratedPath = AssetPackageName.ToString() / FPackageName::GetGeneratedPackageSubPath();
+
+			FARFilter GeneratedFilter;
+			GeneratedFilter.bRecursivePaths = true;
+			GeneratedFilter.bIncludeOnlyOnDiskAssets = true;
+			GeneratedFilter.PackagePaths.Add(FName(*GeneratedPath));
+
+			CurrentRegistrySource->GetAssets(GeneratedFilter, GeneratedAssets);
+		}
+
+		if (GeneratedAssets.Num() > 0)
+		{
+			// Compute children's total size before adding nodes
+			SIZE_T ChildrenTotalSize = 0;
+			for (const FAssetData& GenAsset : GeneratedAssets)
+			{
+				ChildrenTotalSize += GetExternalPackageSize(GenAsset);
+			}
+
+			// Put all the generated packages in a single "Generated Packages" container
+			FTreeMapNodeDataRef GeneratedGroupNode = MakeShareable(new FTreeMapNodeData());
+			ChildTreeMapNode->Children.Add(GeneratedGroupNode);
+			GeneratedGroupNode->Parent = &ChildTreeMapNode.Get();
+			GeneratedGroupNode->Name = FString::Printf(TEXT("Generated Packages  (%d %s, %s)"), GeneratedAssets.Num(), GeneratedAssets.Num() == 1 ? TEXT("package") : TEXT("packages"), *SizeMapInternals::MakeBestSizeString(ChildrenTotalSize, true));
+			GeneratedGroupNode->Size = 0.0f;
+
+			for (const FAssetData& GenAsset : GeneratedAssets)
+			{
+				AddExternalPackageNode(GenAsset, GeneratedGroupNode);
+			}
+
+			// Substract the generated packages size from the generator package size
+			{
+				FNodeSizeMapData& GeneratorPackageEntry = NodeSizeMapDataMap.FindChecked(ChildTreeMapNode);
+				GeneratorPackageEntry.ExternalPackageCount = GeneratedAssets.Num();
+				GeneratorPackageEntry.AssetSize -= FMath::Min(GeneratorPackageEntry.AssetSize, ChildrenTotalSize);
+			}
+		}
+	}
+}
+
 void SSizeMap::FinalizeNodesRecursively(TSharedPtr<FTreeMapNodeData>& Node, const TSharedPtr<FTreeMapNodeData>& SharedRootNode, int32& TotalAssetCount, SIZE_T& TotalSize, bool& bAnyUnknownSizes)
 {
 	// Process children first, so we can get the totals for the root node and shared nodes
@@ -741,6 +875,11 @@ void SSizeMap::FinalizeNodesRecursively(TSharedPtr<FTreeMapNodeData>& Node, cons
 		// Tree root is always auto-sized
 		Node->Size = 0.0f;
 	}
+	else if (!NodeSizeMapDataMap.Contains(Node.ToSharedRef()))
+	{
+		// Container node without asset data
+		Node->Size = 0.0f;
+	}
 	else
 	{
 		// Make a copy as the map may get resized
@@ -780,7 +919,21 @@ void SSizeMap::FinalizeNodesRecursively(TSharedPtr<FTreeMapNodeData>& Node, cons
 		{
 			Node->LogicalName = NodeSizeMapData.AssetData.PackageName.ToString();
 		}
-		
+
+		// For generated packages, use the package short name instead of the asset name
+		FString DisplayName;
+		{
+			const FString PackageNameStr = NodeSizeMapData.AssetData.PackageName.ToString();
+			if (FPackageName::IsUnderGeneratedPackageSubPath(PackageNameStr))
+			{
+				DisplayName = FPackageName::GetShortName(PackageNameStr);
+			}
+			else
+			{
+				DisplayName = NodeSizeMapData.AssetData.AssetName.ToString();
+			}
+		}
+
 		if (Node->IsLeafNode())
 		{
 			Node->CenterText = SizeMapInternals::MakeBestSizeString(NodeSizeMapData.AssetSize, NodeSizeMapData.bHasKnownSize);
@@ -806,7 +959,7 @@ void SSizeMap::FinalizeNodesRecursively(TSharedPtr<FTreeMapNodeData>& Node, cons
 			{
 				// "Asset name"
 				// "Asset type"
-				Node->Name = NodeSizeMapData.AssetData.AssetName.ToString();
+				Node->Name = DisplayName;
 				Node->Name2 = NodeSizeMapData.AssetData.AssetClassPath.ToString();
 			}
 		}
@@ -815,19 +968,31 @@ void SSizeMap::FinalizeNodesRecursively(TSharedPtr<FTreeMapNodeData>& Node, cons
 			// Container nodes are always auto-sized
 			Node->Size = 0.0f;
 
+			const FString SizeString = SizeMapInternals::MakeBestSizeString(SubtreeSize + NodeSizeMapData.AssetSize, !bAnyUnknownSizesInSubtree && NodeSizeMapData.bHasKnownSize);
+
+			// Append external/generated package count when present
+			FString ExternalPackageSuffix;
+			if (NodeSizeMapData.ExternalPackageCount > 0)
+			{
+				const bool bIsCooked = CurrentRegistrySource && !CurrentRegistrySource->bIsEditor;
+				ExternalPackageSuffix = FString::Printf(TEXT(", %d %s %s"), NodeSizeMapData.ExternalPackageCount, bIsCooked ? TEXT("generated") : TEXT("external"), NodeSizeMapData.ExternalPackageCount == 1 ? TEXT("package") : TEXT("packages"));
+			}
+
 			if (PrimaryAssetId.IsValid())
 			{
-				Node->Name = FString::Printf(TEXT("%s  (%s)"),
+				Node->Name = FString::Printf(TEXT("%s  (%s%s)"),
 					*PrimaryAssetId.ToString(),
-					*SizeMapInternals::MakeBestSizeString(SubtreeSize + NodeSizeMapData.AssetSize, !bAnyUnknownSizesInSubtree && NodeSizeMapData.bHasKnownSize));
+					*SizeString,
+					*ExternalPackageSuffix);
 			}
 			else
 			{
-				// "Asset name (asset type, size)"
-				Node->Name = FString::Printf(TEXT("%s  (%s, %s)"),
-					*NodeSizeMapData.AssetData.AssetName.ToString(),
+				// "Asset name (asset type, size, N external/Generated Package(s))"
+				Node->Name = FString::Printf(TEXT("%s  (%s, %s%s)"),
+					*DisplayName,
 					*NodeSizeMapData.AssetData.AssetClassPath.ToString(),
-					*SizeMapInternals::MakeBestSizeString(SubtreeSize + NodeSizeMapData.AssetSize, !bAnyUnknownSizesInSubtree && NodeSizeMapData.bHasKnownSize));
+					*SizeString,
+					*ExternalPackageSuffix);
 			}
 
 			const bool bNeedsSelfNode = NodeSizeMapData.AssetSize > 0;
@@ -1094,6 +1259,15 @@ void SSizeMap::ShowAssetAuditForReferencedAssets() const
 	}
 }
 
+void SSizeMap::CopyPackageName() const
+{
+	const FNodeSizeMapData* NodeSizeMapData = GetCurrentSizeMapData(true);
+	if (NodeSizeMapData)
+	{
+		FPlatformApplicationMisc::ClipboardCopy(*NodeSizeMapData->AssetData.PackageName.ToString());
+	}
+}
+
 void SSizeMap::GetReferencedPackages(TSet<FName>& OutPackageNames, TSharedPtr<FTreeMapNodeData> RootNode, bool bRecurse) const
 {
 	if (!RootNode.IsValid())
@@ -1180,52 +1354,85 @@ void SSizeMap::OnTreeMapNodeRightClicked(FTreeMapNodeData& TreeMapNodeData, cons
 {
 	CurrentSelectedNode = TreeMapNodeData.AsShared();
 
+	// Determine the node type to show appropriate context menu entries
+	const FNodeSizeMapData* NodeData = NodeSizeMapDataMap.Find(TreeMapNodeData.AsShared());
+	const bool bIsContainerNode = (NodeData == nullptr);
+	bool bIsExternalPackage = false;
+	bool bIsGeneratedPackage = false;
+	if (NodeData)
+	{
+		const FString PackageNameStr = NodeData->AssetData.PackageName.ToString();
+		bIsGeneratedPackage = FPackageName::IsUnderGeneratedPackageSubPath(PackageNameStr);
+		bIsExternalPackage = !bIsGeneratedPackage &&
+			(PackageNameStr.Contains(FPackagePath::GetExternalActorsFolderName()) || 
+			 PackageNameStr.Contains(FPackagePath::GetExternalObjectsFolderName()));
+	}
+	const bool bIsRegularAsset = NodeData && !bIsExternalPackage && !bIsGeneratedPackage;
+
 	// Construct menu, this can't support keybinds
 	FMenuBuilder MenuBuilder(/*bInShouldCloseWindowAfterMenuSelection=*/ true, Commands);
 
 	MenuBuilder.BeginSection(TEXT("Asset"), LOCTEXT("AssetSectionLabel", "Asset"));
 	{
-		MenuBuilder.AddMenuEntry(FGlobalEditorCommonCommands::Get().FindInContentBrowser);
-		MenuBuilder.AddMenuEntry(FAssetManagerEditorCommands::Get().OpenSelectedInAssetEditor);
+		if (bIsRegularAsset)
+		{
+			MenuBuilder.AddMenuEntry(FGlobalEditorCommonCommands::Get().FindInContentBrowser);
+			MenuBuilder.AddMenuEntry(FAssetManagerEditorCommands::Get().OpenSelectedInAssetEditor);
+		}
+
+		if (!bIsContainerNode)
+		{
+			MenuBuilder.AddMenuEntry(FAssetManagerEditorCommands::Get().CopyPaths);
+		}
 	}
 	MenuBuilder.EndSection();
 
 	MenuBuilder.BeginSection(TEXT("References"), LOCTEXT("ReferencesSectionLabel", "References"));
 	{
-		MenuBuilder.AddMenuEntry(FAssetManagerEditorCommands::Get().ViewReferences);
-		MenuBuilder.AddMenuEntry(FAssetManagerEditorCommands::Get().ViewAssetAudit);
-		MenuBuilder.AddMenuEntry(FAssetManagerEditorCommands::Get().AuditReferencedObjects);
-
-		ICollectionManager& CollectionManager = FCollectionManagerModule::GetModule().Get();
+		if (bIsRegularAsset || bIsExternalPackage)
+		{
+			MenuBuilder.AddMenuEntry(FAssetManagerEditorCommands::Get().ViewReferences);
+			MenuBuilder.AddMenuEntry(FAssetManagerEditorCommands::Get().ViewAssetAudit);
+		}
 
-		TArray<TSharedPtr<ICollectionContainer>> CollectionContainers;
-		CollectionManager.GetVisibleCollectionContainers(CollectionContainers);
+		if (!bIsGeneratedPackage)
+		{
+			MenuBuilder.AddMenuEntry(FAssetManagerEditorCommands::Get().AuditReferencedObjects);
+		}
 
-		if (!CollectionContainers.IsEmpty())
+		if (bIsRegularAsset)
 		{
-			FNewMenuDelegate SubMenu;
-			if (CollectionContainers.Num() == 1)
+			ICollectionManager& CollectionManager = FCollectionManagerModule::GetModule().Get();
+
+			TArray<TSharedPtr<ICollectionContainer>> CollectionContainers;
+			CollectionManager.GetVisibleCollectionContainers(CollectionContainers);
+
+			if (!CollectionContainers.IsEmpty())
 			{
-				if (CollectionContainers[0] == CollectionManager.GetProjectCollectionContainer())
+				FNewMenuDelegate SubMenu;
+				if (CollectionContainers.Num() == 1)
 				{
-					// Use the command version to show key bindings.
-					SubMenu = FNewMenuDelegate::CreateSP(this, &SSizeMap::GetMakeCollectionWithDependenciesShareTypeSubMenu);
+					if (CollectionContainers[0] == CollectionManager.GetProjectCollectionContainer())
+					{
+						// Use the command version to show key bindings.
+						SubMenu = FNewMenuDelegate::CreateSP(this, &SSizeMap::GetMakeCollectionWithDependenciesShareTypeSubMenu);
+					}
+					else
+					{
+						SubMenu = FNewMenuDelegate::CreateSP(this, &SSizeMap::GetMakeCollectionWithDependenciesShareTypeSubMenu, CollectionContainers[0]);
+					}
 				}
 				else
 				{
-					SubMenu = FNewMenuDelegate::CreateSP(this, &SSizeMap::GetMakeCollectionWithDependenciesShareTypeSubMenu, CollectionContainers[0]);
+					// Create a sub menu to select the collection container.
+					SubMenu = FNewMenuDelegate::CreateSP(this, &SSizeMap::GetMakeCollectionWithDependenciesCollectionContainersSubMenu, MoveTemp(CollectionContainers));
 				}
+				MenuBuilder.AddSubMenu(
+					LOCTEXT("MakeCollectionWithTitle", "Make Collection With References"),
+					LOCTEXT("MakeCollectionWithTooltip", "Makes a collection with all assets referenced by this node."),
+					MoveTemp(SubMenu)
+				);
 			}
-			else
-			{
-				// Create a sub menu to select the collection container.
-				SubMenu = FNewMenuDelegate::CreateSP(this, &SSizeMap::GetMakeCollectionWithDependenciesCollectionContainersSubMenu, MoveTemp(CollectionContainers));
-			}
-			MenuBuilder.AddSubMenu(
-				LOCTEXT("MakeCollectionWithTitle", "Make Collection With References"),
-				LOCTEXT("MakeCollectionWithTooltip", "Makes a collection with all assets referenced by this node."),
-				MoveTemp(SubMenu)
-			);
 		}
 	}
 	MenuBuilder.EndSection();
diff --git a/Engine/Plugins/Editor/AssetManagerEditor/Source/AssetManagerEditor/Private/SSizeMap.h b/Engine/Plugins/Editor/AssetManagerEditor/Source/AssetManagerEditor/Private/SSizeMap.h
--- a/Engine/Plugins/Editor/AssetManagerEditor/Source/AssetManagerEditor/Private/SSizeMap.h
+++ b/Engine/Plugins/Editor/AssetManagerEditor/Source/AssetManagerEditor/Private/SSizeMap.h
@@ -77,10 +77,13 @@ protected:
 	struct FNodeSizeMapData
 	{
 		/** How big the asset is */
-		SIZE_T AssetSize;
+		SIZE_T AssetSize = 0;
 
 		/** Whether it has a known size or not */
-		bool bHasKnownSize;
+		bool bHasKnownSize = false;
+
+		/** Number of external/generated packages discovered for this asset */
+		int32 ExternalPackageCount = 0;
 
 		/** Data from the asset registry about this asset */
 		FAssetData AssetData;
@@ -97,6 +100,9 @@ protected:
 	/** Recursively discovers and loads dependent assets, building up a tree map node hierarchy */
 	void GatherDependenciesRecursively(TSharedPtr<class FAssetThumbnailPool>& AssetThumbnailPool, TMap<FAssetIdentifier, TSharedPtr<class FTreeMapNodeData>>& VisitedAssetIdentifiers, const TArray<FAssetIdentifier>& AssetIdentifiers, const FPrimaryAssetId& FilterPrimaryAsset, const TSharedPtr<FTreeMapNodeData>& ParentTreeMapNode, TSharedPtr<FTreeMapNodeData>& SharedRootNode, int32& NumAssetsWhichFailedToLoad);
 
+	/** Gather all external/generated packages tied to the provided asset */
+	void GatherExternalPackages(const FAssetData CurrentAssetData, FName AssetPackageName, TSharedRef<FTreeMapNodeData> ChildTreeMapNode, TMap<FAssetIdentifier, TSharedPtr<FTreeMapNodeData>>& VisitedAssetIdentifiers);
+
 	/** After the node tree is built up, this function is called to generate nice labels for the nodes and to do a final clean-up pass on the tree */
 	void FinalizeNodesRecursively(TSharedPtr<FTreeMapNodeData>& Node, const TSharedPtr<FTreeMapNodeData>& SharedRootNode, int32& TotalAssetCount, SIZE_T& TotalSize, bool& bAnyUnknownSizes);
 
@@ -123,6 +129,7 @@ protected:
 	void FindReferencesForSelectedAssets() const;
 	void ShowAssetAuditForSelectedAssets() const;
 	void ShowAssetAuditForReferencedAssets() const;
+	void CopyPackageName() const;
 	bool CanMakeCollectionWithDependencies(ECollectionShareType::Type ShareType) const;
 	void MakeCollectionWithDependencies(ECollectionShareType::Type ShareType);
 	bool CanMakeCollectionWithDependencies(TSharedPtr<ICollectionContainer> CollectionContainer, ECollectionShareType::Type ShareType) const;
diff --git a/Engine/Plugins/Editor/AssetManagerEditor/Source/AssetManagerEditor/Public/AssetManagerEditorModule.h b/Engine/Plugins/Editor/AssetManagerEditor/Source/AssetManagerEditor/Public/AssetManagerEditorModule.h
--- a/Engine/Plugins/Editor/AssetManagerEditor/Source/AssetManagerEditor/Public/AssetManagerEditorModule.h
+++ b/Engine/Plugins/Editor/AssetManagerEditor/Source/AssetManagerEditor/Public/AssetManagerEditorModule.h
@@ -203,6 +203,22 @@ public:
 	UE_API void LoadRegistryTimestamp();
 
 	// Functions to forward on to the RegistryState, either the owned one or the editor's global State.
+	FAssetData GetAssetByPackageName(FName PackageName) const
+	{
+		if (bIsEditor)
+		{
+			TArray<FAssetData> Assets;
+			AssetRegistryPrivate->GetAssetsByPackageName(PackageName, Assets, true /* bIncludeOnlyOnDiskAssets */);
+			return Assets.Num() > 0 ? Assets[0] : FAssetData();
+		}
+		else if (RegistryStatePrivate)
+		{
+			TArray<const FAssetData*> Assets = RegistryStatePrivate->CopyAssetsByPackageName(PackageName);
+			return Assets.Num() > 0 ? *Assets[0] : FAssetData();
+		}
+		return FAssetData();
+	}
+
 	FAssetData GetAssetByObjectPath(const FSoftObjectPath& ObjectPath) const
 	{
 		if (bIsEditor)
@@ -292,6 +308,58 @@ public:
 		}
 	}
 
+	bool GetAssets(const FARFilter& Filter, TArray<FAssetData>& OutAssets) const
+	{
+		if (bIsEditor)
+		{
+			return AssetRegistryPrivate->GetAssets(Filter, OutAssets);
+		}
+		else if (RegistryStatePrivate)
+		{
+			FARCompiledFilter CompiledFilter;
+			CompiledFilter.PackageNames.Append(Filter.PackageNames);
+			CompiledFilter.bIncludeOnlyOnDiskAssets = Filter.bIncludeOnlyOnDiskAssets;
+			
+			if (Filter.PackagePaths.Num() > 0)
+			{
+				if (Filter.bRecursivePaths)
+				{
+					// Gather all sub-paths by enumerating known paths in the registry
+					TSet<FName> AllPaths;
+					RegistryStatePrivate->EnumerateAllPaths([&AllPaths](FName Path)
+					{
+						AllPaths.Add(Path);
+					});
+
+					// Filter the paths to keep only those matching the requested paths
+					for (const FName& RequestedPath : Filter.PackagePaths)
+					{
+						const FString RequestedPathStr = RequestedPath.ToString() / TEXT("");
+						for (const FName& CandidatePath : AllPaths)
+						{
+							const FString CandidatePathStr = CandidatePath.ToString();
+							if (CandidatePath == RequestedPath || CandidatePathStr.StartsWith(RequestedPathStr))
+							{
+								CompiledFilter.PackagePaths.Add(CandidatePath);
+							}
+						}
+					}
+				}
+				else
+				{
+					for (const FName& Path : Filter.PackagePaths)
+					{
+						CompiledFilter.PackagePaths.Add(Path);
+					}
+				}
+			}
+
+			RegistryStatePrivate->GetAssets(CompiledFilter, {}, OutAssets);
+			return true;
+		}
+		return false;
+	}
+
 };
 
 /**
