[AssetManager] Display asset name along with asset id for FPrimaryAssetId

[Changed] Customization to show the asset name on top of the primary asset id picker when a thumbnail is present  (giving more vertical real-estate).
[Changed] Customization to show the asset name before the primary asset id picker when a thumbnail is not present.
[Added] Option to add a different tooltip than the name (useful in the second change mentioned above)
[Added] Function that returns the asset name
[Fixed] Small bug with not passing a IsEnabled attribute to the clear button

#14252 tag

#Synced-CL 50983539

diff --git a/Engine/Plugins/Editor/AssetManagerEditor/Source/AssetManagerEditor/Private/AssetManagerEditorModule.cpp b/Engine/Plugins/Editor/AssetManagerEditor/Source/AssetManagerEditor/Private/AssetManagerEditorModule.cpp
--- a/Engine/Plugins/Editor/AssetManagerEditor/Source/AssetManagerEditor/Private/AssetManagerEditorModule.cpp
+++ b/Engine/Plugins/Editor/AssetManagerEditor/Source/AssetManagerEditor/Private/AssetManagerEditorModule.cpp
@@ -152,10 +152,14 @@ TSharedRef<SWidget> IAssetManagerEditorModule::MakePrimaryAssetTypeSelector(FOnG
 	return PropertyCustomizationHelpers::MakePropertyComboBox(nullptr, GetStrings, GetValue, SetValue);
 }
 
-TSharedRef<SWidget> IAssetManagerEditorModule::MakePrimaryAssetIdSelector(FOnGetPrimaryAssetDisplayText OnGetDisplayText, FOnSetPrimaryAssetId OnSetId, bool bAllowClear, 
-	TArray<FPrimaryAssetType> AllowedTypes, TArray<const UClass*> AllowedClasses, TArray<const UClass*> DisallowedClasses)
+TSharedRef<SWidget> IAssetManagerEditorModule::MakePrimaryAssetIdSelector(FPrimaryAssetIdSelectorArgs Args)
 {
-	FOnGetContent OnCreateMenuContent = FOnGetContent::CreateLambda([OnGetDisplayText, OnSetId, bAllowClear, AllowedTypes, AllowedClasses, DisallowedClasses]()
+	FOnGetContent OnCreateMenuContent = FOnGetContent::CreateLambda([
+		OnSetId = Args.OnSetId,
+		bAllowClear = Args.bAllowClear,
+		AllowedTypes = Args.AllowedTypes,
+		AllowedClasses = Args.AllowedClasses,
+		DisallowedClasses = Args.DisallowedClasses]()
 	{
 		FOnShouldFilterAsset AssetFilter = FOnShouldFilterAsset::CreateStatic(&IAssetManagerEditorModule::OnShouldFilterPrimaryAsset, AllowedTypes);
 		FOnSetObject OnSetObject = FOnSetObject::CreateLambda([OnSetId](const FAssetData& AssetData)
@@ -170,7 +174,7 @@ TSharedRef<SWidget> IAssetManagerEditorModule::MakePrimaryAssetIdSelector(FOnGet
 				ensure(AssetId.IsValid());
 			}
 
-			OnSetId.Execute(AssetId);
+			OnSetId.ExecuteIfBound(AssetId);
 		});
 
 		TArray<UFactory*> NewAssetFactories;
@@ -186,7 +190,8 @@ TSharedRef<SWidget> IAssetManagerEditorModule::MakePrimaryAssetIdSelector(FOnGet
 			FSimpleDelegate());
 	});
 
-	TAttribute<FText> OnGetObjectText = TAttribute<FText>::Create(OnGetDisplayText);
+	TAttribute<FText> OnGetObjectText = TAttribute<FText>::Create(Args.OnGetDisplayText);
+	TAttribute<FText> OnGetObjectTooltipText = TAttribute<FText>::Create(Args.OnGetTooltipText);
 
 	return SNew(SComboButton)
 		.OnGetMenuContent(OnCreateMenuContent)
@@ -195,11 +200,28 @@ TSharedRef<SWidget> IAssetManagerEditorModule::MakePrimaryAssetIdSelector(FOnGet
 		[
 			SNew(STextBlock)
 			.Text(OnGetObjectText)
-			.ToolTipText(OnGetObjectText)
+			.ToolTipText(OnGetObjectTooltipText)
 			.Font(FAppStyle::GetFontStyle(TEXT("PropertyWindow.NormalFont")))
 		];
 }
 
+TSharedRef<SWidget> IAssetManagerEditorModule::MakePrimaryAssetIdSelector(FOnGetPrimaryAssetDisplayText OnGetDisplayText, FOnSetPrimaryAssetId OnSetId, bool bAllowClear,
+	TArray<FPrimaryAssetType> AllowedTypes, TArray<const UClass*> AllowedClasses, TArray<const UClass*> DisallowedClasses)
+{
+	return MakePrimaryAssetIdSelector(
+		FPrimaryAssetIdSelectorArgs {
+			.OnGetDisplayText = OnGetDisplayText,
+			.OnGetTooltipText = OnGetDisplayText,
+			.OnSetId = OnSetId,
+			.bAllowClear = bAllowClear,
+			.AllowedTypes = AllowedTypes,
+			.AllowedClasses = AllowedClasses,
+			.DisallowedClasses = DisallowedClasses
+		});
+}
+
+
+
 void IAssetManagerEditorModule::GeneratePrimaryAssetTypeComboBoxStrings(TArray< TSharedPtr<FString> >& OutComboBoxStrings, TArray<TSharedPtr<SToolTip>>& OutToolTips, TArray<bool>& OutRestrictedItems, bool bAllowClear, bool bAllowAll)
 {
 	UAssetManager& AssetManager = UAssetManager::Get();
diff --git a/Engine/Plugins/Editor/AssetManagerEditor/Source/AssetManagerEditor/Private/PrimaryAssetIdCustomization.cpp b/Engine/Plugins/Editor/AssetManagerEditor/Source/AssetManagerEditor/Private/PrimaryAssetIdCustomization.cpp
--- a/Engine/Plugins/Editor/AssetManagerEditor/Source/AssetManagerEditor/Private/PrimaryAssetIdCustomization.cpp
+++ b/Engine/Plugins/Editor/AssetManagerEditor/Source/AssetManagerEditor/Private/PrimaryAssetIdCustomization.cpp
@@ -43,6 +43,37 @@ void FPrimaryAssetIdCustomization::CustomizeHeader(TSharedRef<class IPropertyHan
 		bShowThumbnail = ShowThumbnailString.Len() == 0 || ShowThumbnailString == TEXT("true");
 	}
 
+	TAttribute<bool> bAllowClear = MakeAttributeLambda([this]() -> bool
+		{
+			return !(StructPropertyHandle->GetMetaDataProperty()->PropertyFlags & CPF_NoClear);
+		});
+
+	AllowedClasses = PropertyCustomizationHelpers::GetClassesFromMetadataString(StructPropertyHandle->GetMetaData("AllowedClasses"));
+	DisallowedClasses = PropertyCustomizationHelpers::GetClassesFromMetadataString(StructPropertyHandle->GetMetaData("DisallowedClasses"));
+
+	TSharedRef OptionsBox = SNew(SHorizontalBox)
+		// Use Selected
+		+ SHorizontalBox::Slot()
+		.AutoWidth()
+		.VAlign(VAlign_Center)
+		[
+			PropertyCustomizationHelpers::MakeUseSelectedButton(FSimpleDelegate::CreateSP(this, &FPrimaryAssetIdCustomization::OnUseSelected))
+		]
+		// Browse To
+		+ SHorizontalBox::Slot()
+		.AutoWidth()
+		.VAlign(VAlign_Center)
+		[
+			PropertyCustomizationHelpers::MakeBrowseButton(FSimpleDelegate::CreateSP(this, &FPrimaryAssetIdCustomization::OnBrowseTo))
+		]
+		// Clear
+		+ SHorizontalBox::Slot()
+		.AutoWidth()
+		.VAlign(VAlign_Center)
+		[
+			PropertyCustomizationHelpers::MakeClearButton(FSimpleDelegate::CreateSP(this, &FPrimaryAssetIdCustomization::OnClear), LOCTEXT("ClearButtonToolTipText", "Clear"), bAllowClear)
+		];
+
 	TSharedRef<SHorizontalBox> ValueBox = SNew(SHorizontalBox);
 
 	if (bShowThumbnail)
@@ -51,6 +82,7 @@ void FPrimaryAssetIdCustomization::CustomizeHeader(TSharedRef<class IPropertyHan
 		AssetThumbnail = MakeShareable(new FAssetThumbnail(FAssetData(), ThumbnailSize, ThumbnailSize, StructCustomizationUtils.GetThumbnailPool()));
 		UpdateThumbnail();
 
+		// Thumbnail
 		ValueBox->AddSlot()
 			.AutoWidth()
 			.VAlign(VAlign_Center)
@@ -80,44 +112,78 @@ void FPrimaryAssetIdCustomization::CustomizeHeader(TSharedRef<class IPropertyHan
 					.Visibility(EVisibility::SelfHitTestInvisible)
 				]
 			];
-	}
 
-	// Can the field be cleared
-	const bool bAllowClear = !(StructPropertyHandle->GetMetaDataProperty()->PropertyFlags & CPF_NoClear);
+		TAttribute<FText> AssetNameAttribute = MakeAttributeSP(this, &FPrimaryAssetIdCustomization::GetAssetNameText);
 
-	AllowedClasses = PropertyCustomizationHelpers::GetClassesFromMetadataString(StructPropertyHandle->GetMetaData("AllowedClasses"));
-	DisallowedClasses = PropertyCustomizationHelpers::GetClassesFromMetadataString(StructPropertyHandle->GetMetaData("DisallowedClasses"));
-	
-	ValueBox->AddSlot()
-		.FillWidth(1.0f)
-		.VAlign(VAlign_Center)
-		[
-			IAssetManagerEditorModule::MakePrimaryAssetIdSelector(
-				FOnGetPrimaryAssetDisplayText::CreateSP(this, &FPrimaryAssetIdCustomization::GetDisplayText),
-				FOnSetPrimaryAssetId::CreateSP(this, &FPrimaryAssetIdCustomization::OnIdSelected),
-				bAllowClear, AllowedTypes, AllowedClasses, DisallowedClasses)
-		];
-
-	ValueBox->AddSlot()
-		.AutoWidth()
-		.VAlign(VAlign_Center)
-		[
-			PropertyCustomizationHelpers::MakeUseSelectedButton(FSimpleDelegate::CreateSP(this, &FPrimaryAssetIdCustomization::OnUseSelected))
-		];
+		// Use thumbnail vertical real-estate to show Asset Name on top of the ComboBox
+		ValueBox->AddSlot()
+			.FillWidth(1.0f)
+			.VAlign(VAlign_Center)
+			[
+				SNew(SVerticalBox)
+					// Asset Name
+					+ SVerticalBox::Slot()
+					.VAlign(VAlign_Top)
+					.AutoHeight()
+					.Padding(FMargin(6.f, 6.f, 0.0f, 4.0f))
+					[
+						SNew(STextBlock)
+							.Text(AssetNameAttribute)
+							.ToolTipText(AssetNameAttribute)
+							.OverflowPolicy(ETextOverflowPolicy::Ellipsis)
+							.Font(FAppStyle::GetFontStyle(TEXT("PropertyWindow.BoldFont")))
+					]
+					// PrimaryAssetIdSelector
+					+ SVerticalBox::Slot()
+					.VAlign(VAlign_Center)
+					.FillHeight(1.f)
+					[
+						IAssetManagerEditorModule::MakePrimaryAssetIdSelector(
+							FOnGetPrimaryAssetDisplayText::CreateSP(this, &FPrimaryAssetIdCustomization::GetDisplayText),
+							FOnSetPrimaryAssetId::CreateSP(this, &FPrimaryAssetIdCustomization::OnIdSelected), bAllowClear.Get(), AllowedTypes, AllowedClasses,
+							DisallowedClasses)
+					]
+					// Options
+					+ SVerticalBox::Slot()
+					.Padding(FMargin(4.f, 4.f))
+					.AutoHeight()
+					.VAlign(VAlign_Bottom)
+					[
+						OptionsBox
+					]
+			];
+	}
+	else // No thumbnail
+	{
+		// PrimaryAssetIdSelector
+		ValueBox->AddSlot()
+			.FillWidth(1.0f)
+			.Padding(2.0f)
+			.VAlign(VAlign_Center)
+			[
+					IAssetManagerEditorModule::MakePrimaryAssetIdSelector(
+						{
+							.OnGetDisplayText = FOnGetPrimaryAssetDisplayText::CreateSP(this, &FPrimaryAssetIdCustomization::GetFormattedAssetDisplayText),
+							.OnGetTooltipText = FOnGetPrimaryAssetDisplayText::CreateSP(this, &FPrimaryAssetIdCustomization::GetDisplayText),
+							.OnSetId = FOnSetPrimaryAssetId::CreateSP(this, &FPrimaryAssetIdCustomization::OnIdSelected),
+							.bAllowClear = bAllowClear.Get(),
+							.AllowedTypes = AllowedTypes,
+							.AllowedClasses = AllowedClasses,
+							.DisallowedClasses = DisallowedClasses
+						})
+			];
 
-	ValueBox->AddSlot()
-		.AutoWidth()
-		.VAlign(VAlign_Center)
-		[
-			PropertyCustomizationHelpers::MakeBrowseButton(FSimpleDelegate::CreateSP(this, &FPrimaryAssetIdCustomization::OnBrowseTo))
-		];
+		// Options
+		ValueBox->AddSlot()
+			.VAlign(VAlign_Center)
+			.HAlign(HAlign_Right)
+			.Padding(FMargin(4.0f, 0.0f))
+			.AutoWidth()
+			[
+				OptionsBox
+			];
+	}
 
-	ValueBox->AddSlot()
-		.AutoWidth()
-		.VAlign(VAlign_Center)
-		[
-			PropertyCustomizationHelpers::MakeClearButton(FSimpleDelegate::CreateSP(this, &FPrimaryAssetIdCustomization::OnClear))
-		];
 
 	HeaderRow
 	.NameContent()
@@ -150,19 +216,48 @@ void FPrimaryAssetIdCustomization::OnIdSelected(FPrimaryAssetId AssetId)
 	UpdateThumbnail();
 }
 
-FText FPrimaryAssetIdCustomization::GetDisplayText() const
+FText FPrimaryAssetIdCustomization::GetAssetNameText() const
 {
 	FString StringReference;
-	if (StructPropertyHandle.IsValid())
+
+	FPrimaryAssetId PrimaryAssetId = GetCurrentPrimaryAssetId();
+	if (PrimaryAssetId.IsValid())
 	{
 		StructPropertyHandle->GetValueAsFormattedString(StringReference);
+
+		FAssetData AssetData;
+		if (UAssetManager::Get().GetPrimaryAssetData(PrimaryAssetId, AssetData))
+		{
+			StringReference = AssetData.AssetName.ToString();
+		}
 	}
-	else
+
+	return FText::AsCultureInvariant(StringReference);
+
+}
+
+FText FPrimaryAssetIdCustomization::GetDisplayText() const
+{
+	FString StringReference;
+
+	if (StructPropertyHandle.IsValid())
 	{
-		StringReference = FPrimaryAssetId().ToString();
+		StructPropertyHandle->GetValueAsFormattedString(StringReference);
 	}
 
-	return FText::AsCultureInvariant(StringReference);
+	return !StringReference.IsEmpty() ? FText::AsCultureInvariant(StringReference) : LOCTEXT("EmptyPlaceholderText", "None");
+}
+
+FText FPrimaryAssetIdCustomization::GetFormattedAssetDisplayText() const
+{
+	FText DisplayText = GetDisplayText();
+
+	return DisplayText.EqualTo(LOCTEXT("EmptyPlaceholderText", "None")) ? DisplayText :
+		FText::Format(
+			FText::AsCultureInvariant(TEXT("{0} ( {1} )")),
+			GetAssetNameText(),
+			DisplayText
+		);
 }
 
 FPrimaryAssetId FPrimaryAssetIdCustomization::GetCurrentPrimaryAssetId() const
diff --git a/Engine/Plugins/Editor/AssetManagerEditor/Source/AssetManagerEditor/Private/PrimaryAssetIdCustomization.h b/Engine/Plugins/Editor/AssetManagerEditor/Source/AssetManagerEditor/Private/PrimaryAssetIdCustomization.h
--- a/Engine/Plugins/Editor/AssetManagerEditor/Source/AssetManagerEditor/Private/PrimaryAssetIdCustomization.h
+++ b/Engine/Plugins/Editor/AssetManagerEditor/Source/AssetManagerEditor/Private/PrimaryAssetIdCustomization.h
@@ -30,7 +30,9 @@ private:
 	void OnBrowseTo();
 	void OnClear();
 	void OnUseSelected();
+	FText GetAssetNameText() const;
 	FText GetDisplayText() const;
+	FText GetFormattedAssetDisplayText() const;
 	FPrimaryAssetId GetCurrentPrimaryAssetId() const;
 	void UpdateThumbnail();
 	void OnOpenAssetEditor();
diff --git a/Engine/Plugins/Editor/AssetManagerEditor/Source/AssetManagerEditor/Public/AssetManagerEditorModule.h b/Engine/Plugins/Editor/AssetManagerEditor/Source/AssetManagerEditor/Public/AssetManagerEditorModule.h
--- a/Engine/Plugins/Editor/AssetManagerEditor/Source/AssetManagerEditor/Public/AssetManagerEditorModule.h
+++ b/Engine/Plugins/Editor/AssetManagerEditor/Source/AssetManagerEditor/Public/AssetManagerEditorModule.h
@@ -367,6 +367,18 @@ public:
  */
 class IAssetManagerEditorModule : public IModuleInterface
 {
+public:
+	struct FPrimaryAssetIdSelectorArgs
+	{
+		FOnGetPrimaryAssetDisplayText OnGetDisplayText;
+		FOnGetPrimaryAssetDisplayText OnGetTooltipText;
+		FOnSetPrimaryAssetId OnSetId;
+		bool bAllowClear = true;
+		TArray<FPrimaryAssetType> AllowedTypes;
+		TArray<const UClass*> AllowedClasses;
+		TArray<const UClass*> DisallowedClasses;
+	};
+
 public:
 	/**
 	 * Singleton-like access to this module's interface.  This is just for convenience!
@@ -408,6 +420,12 @@ public:
 		TArray<FPrimaryAssetType> AllowedTypes = TArray<FPrimaryAssetType>(), TArray<const UClass*> AllowedClasses = TArray<const UClass*>(), TArray<const UClass*> DisallowedClasses = TArray<const UClass*>()
 	);
 
+	/**
+	 * Creates a simple version of a Primary Asset Id selector, not bound to a PropertyHandle
+	 * @param PrimaryAssetIdSelectorArgs struct of arguments to build the PrimaryAssetIdSelector
+	 */
+	static UE_API TSharedRef<SWidget> MakePrimaryAssetIdSelector(FPrimaryAssetIdSelectorArgs PrimaryAssetIdSelectorArgs);
+
 	/** Called to get list of valid primary asset types */
 	static UE_API void GeneratePrimaryAssetTypeComboBoxStrings(TArray< TSharedPtr<FString> >& OutComboBoxStrings, TArray<TSharedPtr<SToolTip>>& OutToolTips, TArray<bool>& OutRestrictedItems, bool bAllowClear, bool bAllowAll);
 
