Add way for SizeMap tool to load a developementAssetRegistry for more info.
Add way for the SizeMap tool to select all the assets from a single plugin.
Optimize size map tool for cooked searches with the asset registry, now much snappier when providing a devAR. Optimized via FPathTree

#Synced-CL 51719765

diff --git a/Engine/Plugins/Editor/AssetManagerEditor/Source/AssetManagerEditor/Private/SSizeMap.cpp b/Engine/Plugins/Editor/AssetManagerEditor/Source/AssetManagerEditor/Private/SSizeMap.cpp
--- a/Engine/Plugins/Editor/AssetManagerEditor/Source/AssetManagerEditor/Private/SSizeMap.cpp
+++ b/Engine/Plugins/Editor/AssetManagerEditor/Source/AssetManagerEditor/Private/SSizeMap.cpp
@@ -26,6 +26,7 @@
 #include "Widgets/Input/SComboBox.h"
 #include "Engine/ExternalObjectAndActorDependencyGatherer.h"
 #include "HAL/PlatformApplicationMisc.h"
+#include "Interfaces/IPluginManager.h"
 
 #include UE_INLINE_GENERATED_CPP_BY_NAME(SSizeMap)
 
@@ -141,6 +142,88 @@ FText SSizeMap::GetDependencyTypeText(ESizeMapDependencyType DependencyType) con
 	}
 }
 
+TSharedRef<SWidget> SSizeMap::GenerateSourceComboItem(TSharedPtr<FString> InItem) const
+{
+	return SNew(STextBlock).Text(FText::FromString(*InItem));
+}
+
+void SSizeMap::HandleSourceComboChanged(TSharedPtr<FString> Item, ESelectInfo::Type SelectInfo)
+{
+	if (Item.IsValid())
+	{
+		EditorModule->SetCurrentRegistrySource(*Item);
+	}
+}
+
+FText SSizeMap::GetSourceComboText() const
+{
+	if (CurrentRegistrySource)
+	{
+		return FText::FromString(CurrentRegistrySource->SourceName);
+	}
+	return FText::GetEmpty();
+}
+
+TSharedRef<SWidget> SSizeMap::GeneratePluginComboItem(TSharedPtr<FString> InItem) const
+{
+	const FText DisplayText = InItem->IsEmpty()
+		? LOCTEXT("PluginFilter_None", "None")
+		: FText::FromString(*InItem);
+	return SNew(STextBlock).Text(DisplayText);
+}
+
+void SSizeMap::HandlePluginComboChanged(TSharedPtr<FString> Item, ESelectInfo::Type SelectInfo)
+{
+	if (Item.IsValid())
+	{
+		CurrentPluginSelection = *Item;
+		if (!CurrentPluginSelection.IsEmpty())
+		{
+			SetRootAssetIdentifiersFromPlugin(CurrentPluginSelection);
+		}
+	}
+}
+
+FText SSizeMap::GetPluginComboText() const
+{
+	return CurrentPluginSelection.IsEmpty()
+		? LOCTEXT("PluginFilter_None", "None")
+		: FText::FromString(CurrentPluginSelection);
+}
+
+void SSizeMap::SetRootAssetIdentifiersFromPlugin(const FString& PluginName)
+{
+	TSet<FString> AllPluginNames;
+	// Only support a single plugin being collected right now. Too many will overload the display.
+	AllPluginNames.Add(PluginName);
+
+	// Gather all asset package names from those plugins via the editor asset registry
+	IPluginManager& PluginManager = IPluginManager::Get();
+	FARFilter Filter;
+	Filter.bRecursivePaths = true;
+	for (const FString& Name : AllPluginNames)
+	{
+		TSharedPtr<IPlugin> Plugin = PluginManager.FindPlugin(Name);
+		if (!Plugin.IsValid() || !Plugin->CanContainContent())
+		{
+			continue;
+		}
+		Filter.PackagePaths.Add(FName(FString::Printf(TEXT("/%s/"), *Plugin->GetName())));
+	}
+
+	TArray<FAssetData> FoundAssets;
+	CurrentRegistrySource->GetAssets(Filter, FoundAssets);
+	TSet<FAssetIdentifier> AssetIdentifiers;
+	AssetIdentifiers.Reserve(FoundAssets.Num());
+	for (const FAssetData& Asset : FoundAssets)
+	{
+		AssetIdentifiers.Add(FAssetIdentifier(Asset.PackageName));
+	}
+
+	// Replace existing roots — SetRootAssetIdentifiers clears before setting
+	SetRootAssetIdentifiers(AssetIdentifiers.Array());
+}
+
 void SSizeMap::Construct(const FArguments& InArgs)
 {
 	IAssetManagerEditorModule& ManagerEditorModule = IAssetManagerEditorModule::Get();
@@ -156,6 +239,21 @@ void SSizeMap::Construct(const FArguments& InArgs)
 	DependencyTypeComboList.Add(MakeShared<ESizeMapDependencyType>(ESizeMapDependencyType::Game));
 	DependencyTypeComboList.Add(MakeShared<ESizeMapDependencyType>(ESizeMapDependencyType::EditorOnly));
 
+	// Populate plugin picker list — all enabled plugins that can contain content, sorted by name
+	PluginComboList.Add(MakeShared<FString>(TEXT(""))); // "None" sentinel (empty string)
+	{
+		IPluginManager& PluginManager = IPluginManager::Get();
+		TArray<TSharedRef<IPlugin>> EnabledPlugins = PluginManager.GetEnabledPluginsWithContent();
+		EnabledPlugins.Sort([](const TSharedRef<IPlugin>& A, const TSharedRef<IPlugin>& B)
+		{
+			return A->GetName() < B->GetName();
+		});
+		for (const TSharedRef<IPlugin>& Plugin : EnabledPlugins)
+		{
+			PluginComboList.Add(MakeShared<FString>(Plugin->GetName()));
+		}
+	}
+
 	Commands = MakeShareable(new FUICommandList());
 
 	Commands->MapAction(
@@ -270,6 +368,31 @@ void SSizeMap::Construct(const FArguments& InArgs)
 			.VAlign(VAlign_Center)
 			.Padding(2.f)
 			.AutoWidth()
+			[
+				SNew(STextBlock)
+				.ToolTipText(LOCTEXT("Platform_Tooltip", "Select which platform to display size data for. Platforms are only available if a cooked DevelopmentAssetRegistry.bin exists in Saved/Cooked/[Platform]/Metadata."))
+				.Text(LOCTEXT("PlatformLabel", "Selected Platform: "))
+			]
+			+SHorizontalBox::Slot()
+			.HAlign(HAlign_Right)
+			.VAlign(VAlign_Center)
+			.Padding(2.f)
+			.AutoWidth()
+			[
+				SAssignNew(SourceComboBoxWidget, SComboBox<TSharedPtr<FString>>)
+				.OptionsSource(&SourceComboList)
+				.OnGenerateWidget(this, &SSizeMap::GenerateSourceComboItem)
+				.OnSelectionChanged(this, &SSizeMap::HandleSourceComboChanged)
+				[
+					SNew(STextBlock)
+					.Text(this, &SSizeMap::GetSourceComboText)
+				]
+			]
+			+SHorizontalBox::Slot()
+			.HAlign(HAlign_Right)
+			.VAlign(VAlign_Center)
+			.Padding(2.f)
+			.AutoWidth()
 			[
 				SNew(STextBlock)
 				.Text(LOCTEXT("SizeToDisplay", "Size to Display:"))
@@ -293,12 +416,44 @@ void SSizeMap::Construct(const FArguments& InArgs)
 			]
 		]
 		+SVerticalBox::Slot()
+		.AutoHeight()
+		[
+			SNew(SHorizontalBox)
+			+SHorizontalBox::Slot()
+			.AutoWidth()
+			.VAlign(VAlign_Center)
+			.Padding(2.f)
+			[
+				SNew(STextBlock)
+				.Text(LOCTEXT("PluginFilter_Label", "Plugin Filter:"))
+				.ToolTipText(LOCTEXT("PluginFilter_Tooltip", "Select a plugin to use as the size map root. All assets from the selected plugin will be shown."))
+			]
+			+SHorizontalBox::Slot()
+			.AutoWidth()
+			.VAlign(VAlign_Center)
+			.Padding(2.f)
+			[
+				SAssignNew(PluginComboBoxWidget, SComboBox<TSharedPtr<FString>>)
+				.OptionsSource(&PluginComboList)
+				.OnGenerateWidget(this, &SSizeMap::GeneratePluginComboItem)
+				.OnSelectionChanged(this, &SSizeMap::HandlePluginComboChanged)
+				.ToolTipText(LOCTEXT("PluginFilter_Combo_Tooltip", "Select a plugin. All assets from this plugin will replace the current size map roots."))
+				[
+					SNew(STextBlock)
+					.Text(this, &SSizeMap::GetPluginComboText)
+				]
+			]
+		]
+		+SVerticalBox::Slot()
 		.FillHeight(1.f)
 		[
 			SAssignNew(TreeMapWidget, STreeMap, RootTreeMapNode.ToSharedRef(), nullptr)
 				.OnTreeMapNodeRightClicked(this, &SSizeMap::OnTreeMapNodeRightClicked)
 		]
 	];
+
+	// Refresh registry data to discover any available platform registries (e.g. DevelopmentAssetRegistry.bin)
+	EditorModule->RefreshRegistryData();
 }
 
 FReply SSizeMap::OnKeyDown(const FGeometry& MyGeometry, const FKeyEvent& InKeyEvent)
@@ -346,11 +501,9 @@ FReply SSizeMap::OnDrop(const FGeometry& MyGeometry, const FDragDropEvent& DragD
 
 void SSizeMap::SetRootAssetIdentifiers(const TArray<FAssetIdentifier>& NewRootAssetIdentifiers)
 {
-	if (RootAssetIdentifiers != NewRootAssetIdentifiers)
-	{
-		bMemorySizeCached = false;
-		RootAssetIdentifiers = NewRootAssetIdentifiers;
-	}
+	bMemorySizeCached = false;
+	RootAssetIdentifiers.Reset();
+	RootAssetIdentifiers.Append(NewRootAssetIdentifiers);
 	
 	RefreshMap();
 }
@@ -368,7 +521,20 @@ void SSizeMap::SetCurrentRegistrySource(const FAssetManagerEditorRegistrySource*
 			SizeTypeComboBoxWidget->ClearSelection();
 		}
 	}
-	
+
+	// Refresh the platform combo list to reflect current available registry sources
+	TArray<const FAssetManagerEditorRegistrySource*> AvailableSources;
+	EditorModule->GetAvailableRegistrySources(AvailableSources);
+	SourceComboList.Reset();
+	for (const FAssetManagerEditorRegistrySource* ValidSource : AvailableSources)
+	{
+		SourceComboList.Add(MakeShared<FString>(ValidSource->SourceName));
+	}
+	if (SourceComboBoxWidget.IsValid())
+	{
+		SourceComboBoxWidget->RefreshOptions();
+	}
+
 	RefreshMap();
 }
 
@@ -475,7 +641,7 @@ namespace SizeMapInternals
 	}
 }
 
-void SSizeMap::GatherDependenciesRecursively(TSharedPtr<FAssetThumbnailPool>& InAssetThumbnailPool, TMap<FAssetIdentifier, TSharedPtr<FTreeMapNodeData>>& VisitedAssetIdentifiers, const TArray<FAssetIdentifier>& AssetIdentifiers, const FPrimaryAssetId& FilterPrimaryAsset, const TSharedPtr<FTreeMapNodeData>& Node, TSharedPtr<FTreeMapNodeData>& SharedRootNode, int32& NumAssetsWhichFailedToLoad)
+void SSizeMap::GatherDependenciesRecursively(TSharedPtr<FAssetThumbnailPool>& InAssetThumbnailPool, TMap<FAssetIdentifier, TSharedPtr<FTreeMapNodeData>>& VisitedAssetIdentifiers, const TSet<FAssetIdentifier>& AssetIdentifiers, const FPrimaryAssetId& FilterPrimaryAsset, const TSharedPtr<FTreeMapNodeData>& Node, TSharedPtr<FTreeMapNodeData>& SharedRootNode, int32& NumAssetsWhichFailedToLoad)
 {
 	if (!CurrentRegistrySource->HasRegistry())
 	{
@@ -647,7 +813,7 @@ void SSizeMap::GatherDependenciesRecursively(TSharedPtr<FAssetThumbnailPool>& In
 				// Filter for registry source
 				IAssetManagerEditorModule::Get().FilterAssetIdentifiersForCurrentRegistrySource(References, DependencyQuery, true);
 
-				TArray<FAssetIdentifier> ReferencedAssetIdentifiers;
+				TSet<FAssetIdentifier> ReferencedAssetIdentifiers;
 
 				for (FAssetIdentifier& FoundAssetIdentifier : References)
 				{
@@ -689,7 +855,7 @@ void SSizeMap::GatherDependenciesRecursively(TSharedPtr<FAssetThumbnailPool>& In
 
 				if (AssetPackageName != NAME_None)
 				{
-					if (EditorModule->GetIntegerValueForCustomColumn(NodeSizeMapData.AssetData, GetDefault<USizeMapSettings>()->SizeType, FoundSize))
+					if (EditorModule->GetIntegerValueForCustomColumn(NodeSizeMapData.AssetData, GetDefault<USizeMapSettings>()->SizeType, FoundSize, CurrentRegistrySource))
 					{
 						// If we're reading cooked data, this will fail for dependencies that are editor only. This is fine, they will have 0 size
 						NodeSizeMapData.AssetSize = FoundSize;
@@ -1078,7 +1244,7 @@ void SSizeMap::RefreshMap()
 	// Create a nice name for the tree!
 	if (RootAssetIdentifiers.Num() == 1 && !SharedRootNode.IsValid())
 	{
-		FString OnlyAssetName = RootAssetIdentifiers[0].ToString();
+		FString OnlyAssetName = RootAssetIdentifiers.FindArbitraryElement()->ToString();
 		if (OnlyAssetName.StartsWith(TEXT("/")))
 		{
 			OnlyAssetName = FPackageName::GetShortName(*OnlyAssetName);
diff --git a/Engine/Plugins/Editor/AssetManagerEditor/Source/AssetManagerEditor/Private/SSizeMap.h b/Engine/Plugins/Editor/AssetManagerEditor/Source/AssetManagerEditor/Private/SSizeMap.h
--- a/Engine/Plugins/Editor/AssetManagerEditor/Source/AssetManagerEditor/Private/SSizeMap.h
+++ b/Engine/Plugins/Editor/AssetManagerEditor/Source/AssetManagerEditor/Private/SSizeMap.h
@@ -98,7 +98,7 @@ protected:
 	void OnInitialAssetRegistrySearchComplete();
 
 	/** Recursively discovers and loads dependent assets, building up a tree map node hierarchy */
-	void GatherDependenciesRecursively(TSharedPtr<class FAssetThumbnailPool>& AssetThumbnailPool, TMap<FAssetIdentifier, TSharedPtr<class FTreeMapNodeData>>& VisitedAssetIdentifiers, const TArray<FAssetIdentifier>& AssetIdentifiers, const FPrimaryAssetId& FilterPrimaryAsset, const TSharedPtr<FTreeMapNodeData>& ParentTreeMapNode, TSharedPtr<FTreeMapNodeData>& SharedRootNode, int32& NumAssetsWhichFailedToLoad);
+	void GatherDependenciesRecursively(TSharedPtr<class FAssetThumbnailPool>& AssetThumbnailPool, TMap<FAssetIdentifier, TSharedPtr<class FTreeMapNodeData>>& VisitedAssetIdentifiers, const TSet<FAssetIdentifier>& AssetIdentifiers, const FPrimaryAssetId& FilterPrimaryAsset, const TSharedPtr<FTreeMapNodeData>& ParentTreeMapNode, TSharedPtr<FTreeMapNodeData>& SharedRootNode, int32& NumAssetsWhichFailedToLoad);
 
 	/** Gather all external/generated packages tied to the provided asset */
 	void GatherExternalPackages(const FAssetData CurrentAssetData, FName AssetPackageName, TSharedRef<FTreeMapNodeData> ChildTreeMapNode, TMap<FAssetIdentifier, TSharedPtr<FTreeMapNodeData>>& VisitedAssetIdentifiers);
@@ -161,11 +161,31 @@ protected:
 	TArray<TSharedPtr<ESizeMapDependencyType>> DependencyTypeComboList;
 	TSharedPtr<SComboBox<TSharedPtr<ESizeMapDependencyType>>> DependencyTypeComboBoxWidget;
 
+	/** Platform / Registry Source Combo Box */
+	TSharedRef<SWidget> GenerateSourceComboItem(TSharedPtr<FString> InItem) const;
+	void HandleSourceComboChanged(TSharedPtr<FString> Item, ESelectInfo::Type SelectInfo);
+	FText GetSourceComboText() const;
+
+	TArray<TSharedPtr<FString>> SourceComboList;
+	TSharedPtr<SComboBox<TSharedPtr<FString>>> SourceComboBoxWidget;
+
+	/** Plugin Filter Combo Box — sets plugin assets as root identifiers */
+	TSharedRef<SWidget> GeneratePluginComboItem(TSharedPtr<FString> InItem) const;
+	void HandlePluginComboChanged(TSharedPtr<FString> Item, ESelectInfo::Type SelectInfo);
+	FText GetPluginComboText() const;
+
+	/** Clears root identifiers and replaces them with all assets from PluginName */
+	void SetRootAssetIdentifiersFromPlugin(const FString& PluginName);
+
+	TArray<TSharedPtr<FString>> PluginComboList;
+	TSharedPtr<SComboBox<TSharedPtr<FString>>> PluginComboBoxWidget;
+	FString CurrentPluginSelection; // Empty string = no plugin filter (None)
+
 	/** Our tree map widget */
 	TSharedPtr<class STreeMap> TreeMapWidget;
 
 	/** The assets we were asked to look at */
-	TArray<FAssetIdentifier> RootAssetIdentifiers;
+	TSet<FAssetIdentifier> RootAssetIdentifiers;
 
 	/** Our tree map source data */
 	TSharedPtr<class FTreeMapNodeData> RootTreeMapNode;
diff --git a/Engine/Plugins/Editor/AssetManagerEditor/Source/AssetManagerEditor/Public/AssetManagerEditorModule.h b/Engine/Plugins/Editor/AssetManagerEditor/Source/AssetManagerEditor/Public/AssetManagerEditorModule.h
--- a/Engine/Plugins/Editor/AssetManagerEditor/Source/AssetManagerEditor/Public/AssetManagerEditorModule.h
+++ b/Engine/Plugins/Editor/AssetManagerEditor/Source/AssetManagerEditor/Public/AssetManagerEditorModule.h
@@ -2,6 +2,7 @@
 
 #pragma once
 
+#include "AssetRegistry/PathTree.h"
 #include "CoreMinimal.h"
 #include "Modules/ModuleInterface.h"
 #include "Modules/ModuleManager.h"
@@ -74,6 +75,7 @@ struct FAssetManagerEditorRegistrySource
 private:
 	/** Raw asset registry state, only valid if bIsEditor is false */
 	const FAssetRegistryState* RegistryStatePrivate;
+	FPathTree CachedRegistryStatePathTree;
 	/** Pointer to the real AssetRegistry, only valid if bIsEditor is true */
 	const IAssetRegistry* AssetRegistryPrivate;
 
@@ -127,9 +129,11 @@ public:
 		else
 		{
 			SetRegistryState(Other.RegistryStatePrivate);
+			CachedRegistryStatePathTree = MoveTemp(Other.CachedRegistryStatePathTree);
 		}
 		// Remove the Other's RegistryState without deleting it
 		Other.RegistryStatePrivate = nullptr;
+		Other.CachedRegistryStatePathTree = FPathTree();
 		Other.AssetRegistryPrivate = nullptr;
 		Other.bIsEditor = false;
 
@@ -165,6 +169,11 @@ public:
 		{
 			bIsEditor = false;
 			RegistryStatePrivate = InRegistryState;
+			// Gather all sub-paths by enumerating known paths in the registry
+			RegistryStatePrivate->EnumerateAllPaths([this](FName Path)
+				{
+					CachedRegistryStatePathTree.CachePath(Path, [](FName AddedPath) {});
+				});
 		}
 	}
 
@@ -196,6 +205,7 @@ public:
 		}
 		bIsEditor = false;
 		RegistryStatePrivate = nullptr;
+		CachedRegistryStatePathTree = FPathTree();
 		AssetRegistryPrivate = nullptr;
 	}
 
@@ -324,25 +334,10 @@ public:
 			{
 				if (Filter.bRecursivePaths)
 				{
-					// Gather all sub-paths by enumerating known paths in the registry
-					TSet<FName> AllPaths;
-					RegistryStatePrivate->EnumerateAllPaths([&AllPaths](FName Path)
-					{
-						AllPaths.Add(Path);
-					});
-
-					// Filter the paths to keep only those matching the requested paths
-					for (const FName& RequestedPath : Filter.PackagePaths)
+					// Add the sub-paths of all the input paths to the expanded list
+					for (const FName& PackagePath : Filter.PackagePaths)
 					{
-						const FString RequestedPathStr = RequestedPath.ToString() / TEXT("");
-						for (const FName& CandidatePath : AllPaths)
-						{
-							const FString CandidatePathStr = CandidatePath.ToString();
-							if (CandidatePath == RequestedPath || CandidatePathStr.StartsWith(RequestedPathStr))
-							{
-								CompiledFilter.PackagePaths.Add(CandidatePath);
-							}
-						}
+						CachedRegistryStatePathTree.GetSubPaths(PackagePath, CompiledFilter.PackagePaths);
 					}
 				}
 				else
