Reference Viewer: Add a status bar so that browse to asset can summon its content drawer

#Synced-CL 51808504

diff --git a/Engine/Plugins/Editor/AssetManagerEditor/Source/AssetManagerEditor/AssetManagerEditor.Build.cs b/Engine/Plugins/Editor/AssetManagerEditor/Source/AssetManagerEditor/AssetManagerEditor.Build.cs
--- a/Engine/Plugins/Editor/AssetManagerEditor/Source/AssetManagerEditor/AssetManagerEditor.Build.cs
+++ b/Engine/Plugins/Editor/AssetManagerEditor/Source/AssetManagerEditor/AssetManagerEditor.Build.cs
@@ -54,6 +54,7 @@ public class AssetManagerEditor : ModuleRules
 				"SourceControlWindows",
 				"UncontrolledChangelists",
 				"Projects",
+				"StatusBar",
 			}
 		);
 
diff --git a/Engine/Plugins/Editor/AssetManagerEditor/Source/AssetManagerEditor/Private/AssetManagerEditorModule.cpp b/Engine/Plugins/Editor/AssetManagerEditor/Source/AssetManagerEditor/Private/AssetManagerEditorModule.cpp
--- a/Engine/Plugins/Editor/AssetManagerEditor/Source/AssetManagerEditor/Private/AssetManagerEditorModule.cpp
+++ b/Engine/Plugins/Editor/AssetManagerEditor/Source/AssetManagerEditor/Private/AssetManagerEditorModule.cpp
@@ -783,10 +783,11 @@ TSharedRef<SDockTab> FAssetManagerEditorModule::SpawnReferenceViewerTab(const FS
 {
 	TSharedRef<SDockTab> NewTab = SNew(SDockTab)
 		.TabRole(NomadTab)
-		.OnTabClosed_Raw(this, &FAssetManagerEditorModule::OnReferenceViewerTabClosed)
-		[
-			SAssignNew(ReferenceViewerUI, SReferenceViewer)
-		];
+		.OnTabClosed_Raw(this, &FAssetManagerEditorModule::OnReferenceViewerTabClosed);
+
+	TSharedRef<SReferenceViewer> ReferenceViewerWidget = SNew(SReferenceViewer).ParentTab(NewTab);
+	ReferenceViewerUI = ReferenceViewerWidget;
+	NewTab->SetContent(ReferenceViewerWidget);
 
 	if (!ReferenceViewerTabs.Contains(ReferenceViewerTabName))
 	{
@@ -937,7 +938,9 @@ void FAssetManagerEditorModule::OpenReferenceViewerTab(const TArray<FAssetIdenti
 	if (CurrentAssetTab)
 	{
 		// Initialize Tab
-		CurrentAssetTab->SetContent(SAssignNew(ReferenceViewerUI, SReferenceViewer));
+		TSharedRef<SReferenceViewer> ReferenceViewerWidget = SNew(SReferenceViewer).ParentTab(CurrentAssetTab.ToSharedRef());
+		ReferenceViewerUI = ReferenceViewerWidget;
+		CurrentAssetTab->SetContent(ReferenceViewerWidget);
 		CurrentAssetTab->SetOnTabClosed(SDockTab::FOnTabClosedCallback::CreateRaw(this, &FAssetManagerEditorModule::OnReferenceViewerTabClosed));
 		CurrentAssetTab->SetTabIcon(FSlateIcon(FAppStyle::GetAppStyleSetName(), "ContentBrowser.ReferenceViewer").GetIcon());
 		CurrentAssetTab->SetLabel(FText::FromName(TabLabel));
diff --git a/Engine/Plugins/Editor/AssetManagerEditor/Source/AssetManagerEditor/Private/ReferenceViewer/SReferenceViewer.cpp b/Engine/Plugins/Editor/AssetManagerEditor/Source/AssetManagerEditor/Private/ReferenceViewer/SReferenceViewer.cpp
--- a/Engine/Plugins/Editor/AssetManagerEditor/Source/AssetManagerEditor/Private/ReferenceViewer/SReferenceViewer.cpp
+++ b/Engine/Plugins/Editor/AssetManagerEditor/Source/AssetManagerEditor/Private/ReferenceViewer/SReferenceViewer.cpp
@@ -36,6 +36,7 @@
 #include "Engine/AssetManager.h"
 #include "ReferenceViewer/SReferenceViewerFilterBar.h"
 #include "ReferenceViewer/SReferencedPropertiesNode.h"
+#include "Widgets/Docking/SDockTab.h"
 #include "Widgets/Input/SCheckBox.h"
 #include "Widgets/Input/SComboBox.h"
 #include "Widgets/Input/SComboButton.h"
@@ -45,6 +46,7 @@
 #include "ObjectTools.h"
 #include "Selection.h"
 #include "Interfaces/IPluginManager.h"
+#include "StatusBarSubsystem.h"
 
 #define LOCTEXT_NAMESPACE "ReferenceViewer"
 
@@ -224,6 +226,19 @@ void SReferenceViewer::Construct(const FArguments& InArgs)
 		}
 	}
 
+	// Build the status bar widget if a parent tab was provided
+	TSharedRef<SWidget> StatusBarWidget = SNullWidget::NullWidget;
+	if (InArgs._ParentTab.IsValid())
+	{
+		static int32 StatusBarInstanceCount = 0;
+		const FName StatusBarName = FName(FString::Printf(TEXT("ReferenceViewer_%d"), ++StatusBarInstanceCount));
+		if (UStatusBarSubsystem* StatusBarSubsystem = GEditor->GetEditorSubsystem<UStatusBarSubsystem>())
+		{
+			StatusBarWidget = StatusBarSubsystem->MakeStatusBarWidget(StatusBarName, InArgs._ParentTab.ToSharedRef());
+			StatusBarWidget->SetVisibility(TAttribute<EVisibility>::CreateLambda([ParentTab = InArgs._ParentTab]() { return (ParentTab && ParentTab->GetVisualTabRole() == ETabRole::MajorTab) ? EVisibility::Visible : EVisibility::Collapsed; }));
+		}
+	}
+
 	SAssignNew(FilterWidget, SReferenceViewerFilterBar)
 		.Visibility_Lambda([this]() { return !Settings->GetFiltersEnabled() ? EVisibility::Collapsed : EVisibility::Visible; })
 		.OnConvertItemToAssetData_Lambda([this] (FReferenceNodeInfo& InNodeInfo, FAssetData& OutAssetData) -> bool { 
@@ -672,6 +687,13 @@ void SReferenceViewer::Construct(const FArguments& InArgs)
 				]
 			]
 		]
+
+		// Status bar
+		+SVerticalBox::Slot()
+		.AutoHeight()
+		[
+			StatusBarWidget
+		]
 	];
 
 	SetCanTick(true);
diff --git a/Engine/Plugins/Editor/AssetManagerEditor/Source/AssetManagerEditor/Private/ReferenceViewer/SReferenceViewer.h b/Engine/Plugins/Editor/AssetManagerEditor/Source/AssetManagerEditor/Private/ReferenceViewer/SReferenceViewer.h
--- a/Engine/Plugins/Editor/AssetManagerEditor/Source/AssetManagerEditor/Private/ReferenceViewer/SReferenceViewer.h
+++ b/Engine/Plugins/Editor/AssetManagerEditor/Source/AssetManagerEditor/Private/ReferenceViewer/SReferenceViewer.h
@@ -12,6 +12,7 @@ class ICollectionContainer;
 class IPlugin;
 class FUICommandList;
 class SComboButton;
+class SDockTab;
 class SReferenceViewerFilterBar;
 class SSearchBox;
 class UEdGraph_ReferenceViewer;
@@ -31,7 +32,7 @@ class SReferenceViewer : public SCompoundWidget
 {
 public:
 	SLATE_BEGIN_ARGS( SReferenceViewer ){}
-
+		SLATE_ARGUMENT(TSharedPtr<SDockTab>, ParentTab)
 	SLATE_END_ARGS()
 
 	~SReferenceViewer();
