Change reference viewer status bar naming scheme and fix SDockTab leaking because a sharedptr was captured in the visibility lambda.

#Synced-CL 51808504

diff --git a/Engine/Plugins/Editor/AssetManagerEditor/Source/AssetManagerEditor/Private/ReferenceViewer/SReferenceViewer.cpp b/Engine/Plugins/Editor/AssetManagerEditor/Source/AssetManagerEditor/Private/ReferenceViewer/SReferenceViewer.cpp
--- a/Engine/Plugins/Editor/AssetManagerEditor/Source/AssetManagerEditor/Private/ReferenceViewer/SReferenceViewer.cpp
+++ b/Engine/Plugins/Editor/AssetManagerEditor/Source/AssetManagerEditor/Private/ReferenceViewer/SReferenceViewer.cpp
@@ -230,12 +230,11 @@ void SReferenceViewer::Construct(const FArguments& InArgs)
 	TSharedRef<SWidget> StatusBarWidget = SNullWidget::NullWidget;
 	if (InArgs._ParentTab.IsValid())
 	{
-		static int32 StatusBarInstanceCount = 0;
-		const FName StatusBarName = FName(FString::Printf(TEXT("ReferenceViewer_%d"), ++StatusBarInstanceCount));
+		const FName StatusBarName(FString::Printf(TEXT("ReferenceViewer_0x%llx"), reinterpret_cast<uint64>(this)));
 		if (UStatusBarSubsystem* StatusBarSubsystem = GEditor->GetEditorSubsystem<UStatusBarSubsystem>())
 		{
 			StatusBarWidget = StatusBarSubsystem->MakeStatusBarWidget(StatusBarName, InArgs._ParentTab.ToSharedRef());
-			StatusBarWidget->SetVisibility(TAttribute<EVisibility>::CreateLambda([ParentTab = InArgs._ParentTab]() { return (ParentTab && ParentTab->GetVisualTabRole() == ETabRole::MajorTab) ? EVisibility::Visible : EVisibility::Collapsed; }));
+			StatusBarWidget->SetVisibility(TAttribute<EVisibility>::CreateLambda([ParentTab = TWeakPtr<SDockTab>(InArgs._ParentTab)]() { return (ParentTab.IsValid() && ParentTab.Pin()->GetVisualTabRole() == ETabRole::MajorTab) ? EVisibility::Visible : EVisibility::Collapsed; }));
 		}
 	}
 
