UE: Do not load redirectors when looking at them in the reference viewer

#Synced-CL 52440464

diff --git a/Engine/Plugins/Editor/AssetManagerEditor/Source/AssetManagerEditor/Private/ReferenceViewer/EdGraph_ReferenceViewer.cpp b/Engine/Plugins/Editor/AssetManagerEditor/Source/AssetManagerEditor/Private/ReferenceViewer/EdGraph_ReferenceViewer.cpp
--- a/Engine/Plugins/Editor/AssetManagerEditor/Source/AssetManagerEditor/Private/ReferenceViewer/EdGraph_ReferenceViewer.cpp
+++ b/Engine/Plugins/Editor/AssetManagerEditor/Source/AssetManagerEditor/Private/ReferenceViewer/EdGraph_ReferenceViewer.cpp
@@ -1244,32 +1244,45 @@ UEdGraph_ReferenceViewer::RecursivelyPopulateNodeInfos(bool bInReferencers, cons
 	FAssetData* AssetData = PackageToAssetDataMap.Find(InAssetId.PackageName);
 	if (!bInReferencers && AssetData && AssetData->IsRedirector())
 	{
-		if (UObjectRedirector* Redirector = Cast<UObjectRedirector>(AssetData->GetAsset()))
-		{
-			bIsRedirector = true;
+		bIsRedirector = true;
+
+		IAssetRegistry* AssetRegistry = IAssetRegistry::Get();
+		check(AssetRegistry);
 
-			// We are dealing with a redirector. Let's manually retrieve its Destination Object, and set up its set of nodes explicitly
-			if (const UObject* const DestinationObject = Redirector->DestinationObject)
+		// We are dealing with a redirector. Let's manually retrieve its Destination Object, and set up its set of nodes explicitly
+		FAssetData DestinationObjectAssetData = *AssetData;
+		TSet<FString> AllVisitedObjectPaths;
+		while (DestinationObjectAssetData.IsRedirector())
+		{
+			FString DestinationObject = FPackageName::ExportTextPathToObjectPath(DestinationObjectAssetData.GetTagValueRef<FString>(GET_MEMBER_NAME_CHECKED(UObjectRedirector, DestinationObject)));
+			if (DestinationObject.IsEmpty() || AllVisitedObjectPaths.Contains(DestinationObject))
 			{
-				if (const UPackage* const DestinationObjectPackage = DestinationObject->GetPackage())
-				{
-					const FName& DestinationPackageName = DestinationObjectPackage->GetFName();
-					const FAssetIdentifier DestinationAssetId = FAssetIdentifier::FromString(DestinationPackageName.ToString());
+				// No destination or a loop
+				DestinationObjectAssetData = FAssetData();
+				break;
+			}
 
-					FReferenceNodeInfo& DestinationReferenceNodeInfo = InNodeInfos.FindOrAdd(DestinationAssetId, FReferenceNodeInfo(DestinationAssetId, bInReferencers));
+			AllVisitedObjectPaths.Add(DestinationObject);
+			DestinationObjectAssetData = AssetRegistry->GetAssetByObjectPath(DestinationObject, true);
+		}
 
-					// The Destination Node parent is the Redirector one
-					DestinationReferenceNodeInfo.Parents.Emplace(InAssetId);
+		if (DestinationObjectAssetData.IsValid())
+		{
+			const FName DestinationPackageName = DestinationObjectAssetData.PackageName;
+			const FAssetIdentifier DestinationAssetId = FAssetIdentifier::FromString(DestinationPackageName.ToString());
 
-					// Remove Children from Redirector Node, and just add the Destination Node
-					InNodeInfos[InAssetId].Children.Empty();
-					InNodeInfos[InAssetId].Children.Emplace(DestinationAssetId, EDependencyPinCategory::LinkTypeHard);
-					InNodeInfos[InAssetId].bIsRedirector = true;
+			FReferenceNodeInfo& DestinationReferenceNodeInfo = InNodeInfos.FindOrAdd(DestinationAssetId, FReferenceNodeInfo(DestinationAssetId, bInReferencers));
 
-					// Populate Info, without increasing current depth - we ignore the Redirector
-					RecursivelyPopulateNodeInfos(bInReferencers, { DestinationAssetId }, InNodeInfos, 0, InMaxDepth - InCurrentDepth);
-				}
-			}
+			// The Destination Node parent is the Redirector one
+			DestinationReferenceNodeInfo.Parents.Emplace(InAssetId);
+
+			// Remove Children from Redirector Node, and just add the Destination Node
+			InNodeInfos[InAssetId].Children.Empty();
+			InNodeInfos[InAssetId].Children.Emplace(DestinationAssetId, EDependencyPinCategory::LinkTypeHard);
+			InNodeInfos[InAssetId].bIsRedirector = true;
+
+			// Populate Info, without increasing current depth - we ignore the Redirector
+			RecursivelyPopulateNodeInfos(bInReferencers, { DestinationAssetId }, InNodeInfos, 0, InMaxDepth - InCurrentDepth);
 		}
 	}
 
